Skeletal Reduction, update whitelist to include Mac and Linux

diff --git a/Engine/Plugins/Experimental/SkeletalReduction/SkeletalReduction.uplugin b/Engine/Plugins/Experimental/SkeletalReduction/SkeletalReduction.uplugin
--- a/Engine/Plugins/Experimental/SkeletalReduction/SkeletalReduction.uplugin
+++ b/Engine/Plugins/Experimental/SkeletalReduction/SkeletalReduction.uplugin
@@ -20,7 +20,7 @@
 			"Name" : "SkeletalMeshReduction",
 			"Type" : "Editor",
 			"LoadingPhase" : "Default",
-			 "WhitelistPlatforms": [ "Win64" ]
+			"WhitelistPlatforms": [ "Win64", "Mac", "Linux"]
 		}
 	]
 }
\ No newline at end of file
