SkeletalSimplifier:   Error during incremental UE4Editor Linux, non-portable path

diff --git a/Engine/Plugins/Experimental/SkeletalReduction/Source/Private/SkeletalSimplifierVertex.h b/Engine/Plugins/Experimental/SkeletalReduction/Source/Private/SkeletalSimplifierVertex.h
--- a/Engine/Plugins/Experimental/SkeletalReduction/Source/Private/SkeletalSimplifierVertex.h
+++ b/Engine/Plugins/Experimental/SkeletalReduction/Source/Private/SkeletalSimplifierVertex.h
@@ -5,7 +5,7 @@
 #include "SkeletalSimplifierLinearAlgebra.h"
 
 #include "Math/Color.h"     // FLinearColor
-#include "Math/Vector2d.h"  // FVector2d
+#include "Math/Vector2D.h"  // FVector2d
 #include "Math/Vector.h"    // FVector
 #include "MeshBuild.h"      // ApprxEquals for normals & UVs
 #include "Templates/UnrealTemplate.h"
