Make sure we can reimport specific content on a reduce base LOD

diff --git a/Engine/Plugins/Experimental/SkeletalReduction/Source/Private/SkeletalMeshReductionPlugin.cpp b/Engine/Plugins/Experimental/SkeletalReduction/Source/Private/SkeletalMeshReductionPlugin.cpp
--- a/Engine/Plugins/Experimental/SkeletalReduction/Source/Private/SkeletalMeshReductionPlugin.cpp
+++ b/Engine/Plugins/Experimental/SkeletalReduction/Source/Private/SkeletalMeshReductionPlugin.cpp
@@ -1544,6 +1544,13 @@ void FQuadricSkeletalMeshReduction::ReduceSkeletalMesh(USkeletalMesh& SkeletalMe
 		{
 			delete Old;
 		}
+		else if(bReducingSourceModel)
+		{
+			//In case we reduce the source model we want to keep the original import data
+			FSkeletalMeshImportData RawMesh;
+			SrcModel->RawSkeletalMeshBulkData.LoadRawMesh(RawMesh);
+			LODModels[LODIndex]->RawSkeletalMeshBulkData.SaveRawMesh(RawMesh);
+		}
 	}
 
 	
