Keep the cloth binding when reducing a source model with himself has the source LOD

diff --git a/Engine/Plugins/Experimental/SkeletalReduction/Source/Private/SkeletalMeshReductionPlugin.cpp b/Engine/Plugins/Experimental/SkeletalReduction/Source/Private/SkeletalMeshReductionPlugin.cpp
--- a/Engine/Plugins/Experimental/SkeletalReduction/Source/Private/SkeletalMeshReductionPlugin.cpp
+++ b/Engine/Plugins/Experimental/SkeletalReduction/Source/Private/SkeletalMeshReductionPlugin.cpp
@@ -22,7 +22,7 @@
 #include "SkeletalSimplifier.h"
 #include "SkeletalMeshReductionSkinnedMesh.h"
 #include "Stats/StatsMisc.h"
-
+#include "Assets/ClothingAsset.h"
 
 
 #define LOCTEXT_NAMESPACE "SkeletalMeshReduction"
@@ -1448,8 +1448,45 @@ void FQuadricSkeletalMeshReduction::ReduceSkeletalMesh(USkeletalMesh& SkeletalMe
 			UE_LOG(LogSkeletalMeshReduction, Warning, TEXT("Building LOD %d - Invalid Base LOD entered. Using Base LOD 0 instead"), LODIndex);
 		}
 	}
-	//Reducing Base LOD, we need to use the temporary data so it can be iterative
+	
+	auto FillClothingData = [&SkeletalMeshResource, &LODIndex](int32 &EnableSectionNumber, TArray<bool> &SectionStatus)
+	{
+		EnableSectionNumber = 0;
+		SectionStatus.Empty();
+		if (SkeletalMeshResource.LODModels.IsValidIndex(LODIndex))
+		{
+			int32 SectionNumber = SkeletalMeshResource.LODModels[LODIndex].Sections.Num();
+			SectionStatus.Reserve(SectionNumber);
+			for (int32 SectionIndex = 0; SectionIndex < SectionNumber; ++SectionIndex)
+			{
+				SectionStatus.Add(!SkeletalMeshResource.LODModels[LODIndex].Sections[SectionIndex].bDisabled);
+				if (SectionStatus[SectionIndex])
+				{
+					EnableSectionNumber++;
+				}
+			}
+		}
+	};
+
+	// Unbind any existing clothing assets before we reimport the geometry
+	TArray<ClothingAssetUtils::FClothingAssetMeshBinding> ClothingBindings;
+	//Store the clothBinding
+	ClothingAssetUtils::GetMeshClothingAssetBindings(&SkeletalMesh, ClothingBindings, LODIndex);
+	//Get a map of enable/disable sections
+	int32 OriginalSectionNumberBeforeReduction = 0;
+	TArray<bool> OriginalSectionEnableBeforeReduction;
+	FillClothingData(OriginalSectionNumberBeforeReduction, OriginalSectionEnableBeforeReduction);
+	//Unbind the Cloth for this LOD before we reduce it, we will put back the cloth after the reduction, if it still match the sections
+	for (ClothingAssetUtils::FClothingAssetMeshBinding& Binding : ClothingBindings)
+	{
+		if (Binding.LODIndex == LODIndex)
+		{
+			Binding.Asset->UnbindFromSkeletalMesh(&SkeletalMesh, Binding.LODIndex);
+		}
+	}
+
 	bool bReducingSourceModel = false;
+	//Reducing Base LOD, we need to use the temporary data so it can be iterative
 	if (BaseLOD == LODIndex && SkelResource->OriginalReductionSourceMeshData.IsValidIndex(BaseLOD) && !SkelResource->OriginalReductionSourceMeshData[BaseLOD]->IsEmpty())
 	{
 		TMap<FString, TArray<FMorphTargetDelta>> TempLODMorphTargetData;
@@ -1641,6 +1678,43 @@ void FQuadricSkeletalMeshReduction::ReduceSkeletalMesh(USkeletalMesh& SkeletalMe
 		SkeletalMesh.bHasBeenSimplified = true;
 	}
 
+	//Get the number of enabled section
+	int32 SectionNumberAfterReduction = 0;
+	TArray<bool> SectionEnableAfterReduction;
+	FillClothingData(SectionNumberAfterReduction, SectionEnableAfterReduction);
+
+	//Put back the clothing for this newly reduce LOD only if the section count match.
+	if (ClothingBindings.Num() > 0 && OriginalSectionNumberBeforeReduction == SectionNumberAfterReduction)
+	{
+		TArray<int32> RemapSectionIndex;
+		int32 SectionIndexTest = 0;
+		for (int32 SectionIndexRef = 0; SectionIndexRef < OriginalSectionEnableBeforeReduction.Num(); SectionIndexRef++)
+		{
+			int32& RemapValue = RemapSectionIndex.Add_GetRef(INDEX_NONE);
+			if (!OriginalSectionEnableBeforeReduction[SectionIndexRef])
+			{
+				continue;
+			}
+			for (; SectionIndexTest <= SectionIndexRef; SectionIndexTest++)
+			{
+				if (SectionEnableAfterReduction.IsValidIndex(SectionIndexTest) && SectionEnableAfterReduction[SectionIndexTest])
+				{
+					RemapValue = SectionIndexTest++;
+					break;
+				}
+			}
+		}
+
+		for (ClothingAssetUtils::FClothingAssetMeshBinding& Binding : ClothingBindings)
+		{
+			int32 RemapBindingSectionIndex = RemapSectionIndex[Binding.SectionIndex];
+			if (RemapBindingSectionIndex != INDEX_NONE && Binding.LODIndex == LODIndex && NewModel->Sections.IsValidIndex(RemapBindingSectionIndex))
+			{
+				Binding.Asset->BindToSkeletalMesh(&SkeletalMesh, Binding.LODIndex, RemapBindingSectionIndex, Binding.AssetInternalLodIndex, false);
+			}
+		}
+	}
+
 	SkeletalMesh.CalculateRequiredBones(SkeletalMeshResource.LODModels[LODIndex], SkeletalMesh.RefSkeleton, &BonesToRemove);
 }
 
diff --git a/Engine/Plugins/Experimental/SkeletalReduction/Source/SkeletalMeshReduction.Build.cs b/Engine/Plugins/Experimental/SkeletalReduction/Source/SkeletalMeshReduction.Build.cs
--- a/Engine/Plugins/Experimental/SkeletalReduction/Source/SkeletalMeshReduction.Build.cs
+++ b/Engine/Plugins/Experimental/SkeletalReduction/Source/SkeletalMeshReduction.Build.cs
@@ -55,7 +55,8 @@ namespace UnrealBuildTool.Rules
                     "UnrealEd",
                     "AnimationModifiers",
                     "MeshBoneReduction",
-                    "QuadricMeshReduction"
+                    "QuadricMeshReduction",
+                    "ClothingSystemRuntime"
 					// ... add private dependencies that you statically link with here ...
                     // QuadricMeshReduction is only for testing
 				}
