Fix the skeletalmesh ActiveReduction function to also look for absolute criterions

diff --git a/Engine/Plugins/Experimental/SkeletalReduction/Source/Private/SkeletalMeshReductionPlugin.cpp b/Engine/Plugins/Experimental/SkeletalReduction/Source/Private/SkeletalMeshReductionPlugin.cpp
--- a/Engine/Plugins/Experimental/SkeletalReduction/Source/Private/SkeletalMeshReductionPlugin.cpp
+++ b/Engine/Plugins/Experimental/SkeletalReduction/Source/Private/SkeletalMeshReductionPlugin.cpp
@@ -80,6 +80,14 @@ public:
 				return ReductionSettings.NumOfTrianglesPercentage < Threshold_One || ReductionSettings.NumOfVertPercentage < Threshold_One;
 			}
 			break;
+			//Absolute count is consider has being always reduced
+			case SkeletalMeshTerminationCriterion::SMTC_AbsNumOfVerts:
+			case SkeletalMeshTerminationCriterion::SMTC_AbsNumOfTriangles:
+			case SkeletalMeshTerminationCriterion::SMTC_AbsTriangleOrVert:
+			{
+				return true;
+			}
+			break;
 		}
 
 		return false;
