Fix crash when generating LODs

diff --git a/Engine/Plugins/Experimental/SkeletalReduction/Source/Private/SkeletalMeshReductionPlugin.cpp b/Engine/Plugins/Experimental/SkeletalReduction/Source/Private/SkeletalMeshReductionPlugin.cpp
--- a/Engine/Plugins/Experimental/SkeletalReduction/Source/Private/SkeletalMeshReductionPlugin.cpp
+++ b/Engine/Plugins/Experimental/SkeletalReduction/Source/Private/SkeletalMeshReductionPlugin.cpp
@@ -1389,11 +1389,15 @@ void FQuadricSkeletalMeshReduction::ReduceSkeletalMesh(USkeletalMesh& SkeletalMe
 	//If the Current LOD is an import from file
 	bool bOldLodWasFromFile = SkeletalMesh.IsValidLODIndex(LODIndex) && SkeletalMesh.GetLODInfo(LODIndex)->bHasBeenSimplified == false;
 
+	//True if the LOD is added by this reduction
+	bool bLODModelAdded = false;
+
 	// Insert a new LOD model entry if needed.
 	if (LODIndex == SkeletalMeshResource.LODModels.Num())
 	{
 		FSkeletalMeshLODModel* ModelPtr = NULL;
 		SkeletalMeshResource.LODModels.Add(ModelPtr);
+		bLODModelAdded = true;
 	}
 
 	// Copy over LOD info from LOD0 if there is no previous info.
@@ -1457,11 +1461,11 @@ void FQuadricSkeletalMeshReduction::ReduceSkeletalMesh(USkeletalMesh& SkeletalMe
 		}
 	}
 	
-	auto FillClothingData = [&SkeletalMeshResource, &LODIndex](int32 &EnableSectionNumber, TArray<bool> &SectionStatus)
+	auto FillClothingData = [&SkeletalMeshResource, &LODIndex, bLODModelAdded](int32 &EnableSectionNumber, TArray<bool> &SectionStatus)
 	{
 		EnableSectionNumber = 0;
 		SectionStatus.Empty();
-		if (SkeletalMeshResource.LODModels.IsValidIndex(LODIndex))
+		if (!bLODModelAdded && SkeletalMeshResource.LODModels.IsValidIndex(LODIndex))
 		{
 			int32 SectionNumber = SkeletalMeshResource.LODModels[LODIndex].Sections.Num();
 			SectionStatus.Reserve(SectionNumber);
@@ -1478,18 +1482,23 @@ void FQuadricSkeletalMeshReduction::ReduceSkeletalMesh(USkeletalMesh& SkeletalMe
 
 	// Unbind any existing clothing assets before we reimport the geometry
 	TArray<ClothingAssetUtils::FClothingAssetMeshBinding> ClothingBindings;
-	//Store the clothBinding
-	ClothingAssetUtils::GetMeshClothingAssetBindings(&SkeletalMesh, ClothingBindings, LODIndex);
 	//Get a map of enable/disable sections
 	int32 OriginalSectionNumberBeforeReduction = 0;
 	TArray<bool> OriginalSectionEnableBeforeReduction;
-	FillClothingData(OriginalSectionNumberBeforeReduction, OriginalSectionEnableBeforeReduction);
-	//Unbind the Cloth for this LOD before we reduce it, we will put back the cloth after the reduction, if it still match the sections
-	for (ClothingAssetUtils::FClothingAssetMeshBinding& Binding : ClothingBindings)
+
+	//Do not play with cloth if the LOD is added
+	if (!bLODModelAdded)
 	{
-		if (Binding.LODIndex == LODIndex)
+		//Store the clothBinding
+		ClothingAssetUtils::GetMeshClothingAssetBindings(&SkeletalMesh, ClothingBindings, LODIndex);
+		FillClothingData(OriginalSectionNumberBeforeReduction, OriginalSectionEnableBeforeReduction);
+		//Unbind the Cloth for this LOD before we reduce it, we will put back the cloth after the reduction, if it still match the sections
+		for (ClothingAssetUtils::FClothingAssetMeshBinding& Binding : ClothingBindings)
 		{
-			Binding.Asset->UnbindFromSkeletalMesh(&SkeletalMesh, Binding.LODIndex);
+			if (Binding.LODIndex == LODIndex)
+			{
+				Binding.Asset->UnbindFromSkeletalMesh(&SkeletalMesh, Binding.LODIndex);
+			}
 		}
 	}
 
@@ -1685,40 +1694,43 @@ void FQuadricSkeletalMeshReduction::ReduceSkeletalMesh(USkeletalMesh& SkeletalMe
 		SkeletalMesh.GetLODInfo(LODIndex)->bHasBeenSimplified = true;
 		SkeletalMesh.bHasBeenSimplified = true;
 	}
-
-	//Get the number of enabled section
-	int32 SectionNumberAfterReduction = 0;
-	TArray<bool> SectionEnableAfterReduction;
-	FillClothingData(SectionNumberAfterReduction, SectionEnableAfterReduction);
-
-	//Put back the clothing for this newly reduce LOD only if the section count match.
-	if (ClothingBindings.Num() > 0 && OriginalSectionNumberBeforeReduction == SectionNumberAfterReduction)
+	
+	if (!bLODModelAdded)
 	{
-		TArray<int32> RemapSectionIndex;
-		int32 SectionIndexTest = 0;
-		for (int32 SectionIndexRef = 0; SectionIndexRef < OriginalSectionEnableBeforeReduction.Num(); SectionIndexRef++)
+		//Get the number of enabled section
+		int32 SectionNumberAfterReduction = 0;
+		TArray<bool> SectionEnableAfterReduction;
+		FillClothingData(SectionNumberAfterReduction, SectionEnableAfterReduction);
+
+		//Put back the clothing for this newly reduce LOD only if the section count match.
+		if (ClothingBindings.Num() > 0 && OriginalSectionNumberBeforeReduction == SectionNumberAfterReduction)
 		{
-			int32& RemapValue = RemapSectionIndex.Add_GetRef(INDEX_NONE);
-			if (!OriginalSectionEnableBeforeReduction[SectionIndexRef])
-			{
-				continue;
-			}
-			for (; SectionIndexTest <= SectionIndexRef; SectionIndexTest++)
+			TArray<int32> RemapSectionIndex;
+			int32 SectionIndexTest = 0;
+			for (int32 SectionIndexRef = 0; SectionIndexRef < OriginalSectionEnableBeforeReduction.Num(); SectionIndexRef++)
 			{
-				if (SectionEnableAfterReduction.IsValidIndex(SectionIndexTest) && SectionEnableAfterReduction[SectionIndexTest])
+				int32& RemapValue = RemapSectionIndex.Add_GetRef(INDEX_NONE);
+				if (!OriginalSectionEnableBeforeReduction[SectionIndexRef])
+				{
+					continue;
+				}
+				for (; SectionIndexTest <= SectionIndexRef; SectionIndexTest++)
 				{
-					RemapValue = SectionIndexTest++;
-					break;
+					if (SectionEnableAfterReduction.IsValidIndex(SectionIndexTest) && SectionEnableAfterReduction[SectionIndexTest])
+					{
+						RemapValue = SectionIndexTest++;
+						break;
+					}
 				}
 			}
-		}
 
-		for (ClothingAssetUtils::FClothingAssetMeshBinding& Binding : ClothingBindings)
-		{
-			int32 RemapBindingSectionIndex = RemapSectionIndex[Binding.SectionIndex];
-			if (RemapBindingSectionIndex != INDEX_NONE && Binding.LODIndex == LODIndex && NewModel->Sections.IsValidIndex(RemapBindingSectionIndex))
+			for (ClothingAssetUtils::FClothingAssetMeshBinding& Binding : ClothingBindings)
 			{
-				Binding.Asset->BindToSkeletalMesh(&SkeletalMesh, Binding.LODIndex, RemapBindingSectionIndex, Binding.AssetInternalLodIndex, false);
+				int32 RemapBindingSectionIndex = RemapSectionIndex[Binding.SectionIndex];
+				if (RemapBindingSectionIndex != INDEX_NONE && Binding.LODIndex == LODIndex && NewModel->Sections.IsValidIndex(RemapBindingSectionIndex))
+				{
+					Binding.Asset->BindToSkeletalMesh(&SkeletalMesh, Binding.LODIndex, RemapBindingSectionIndex, Binding.AssetInternalLodIndex, false);
+				}
 			}
 		}
 	}
