Fix a crash with alternate weights

diff --git a/Engine/Plugins/Experimental/SkeletalReduction/Source/Private/SkeletalMeshReductionPlugin.cpp b/Engine/Plugins/Experimental/SkeletalReduction/Source/Private/SkeletalMeshReductionPlugin.cpp
--- a/Engine/Plugins/Experimental/SkeletalReduction/Source/Private/SkeletalMeshReductionPlugin.cpp
+++ b/Engine/Plugins/Experimental/SkeletalReduction/Source/Private/SkeletalMeshReductionPlugin.cpp
@@ -472,7 +472,7 @@ void FQuadricSkeletalMeshReduction::ConvertToFSkinnedSkeletalMesh( const FSkelet
 		
 		Vertex.TangentX   = WeightedTangentX.GetSafeNormal();
 		Vertex.TangentY   = WeightedTangentY.GetSafeNormal();
-		uint8 WComponent  = Vertex.TangentZ.W;                    // NB: this looks bad.. the component of the vector is a float..
+		float WComponent  = Vertex.TangentZ.W;
 		Vertex.TangentZ   = WeightedTangentZ.GetSafeNormal();
 		Vertex.TangentZ.W = WComponent;
 		Vertex.Position   = WeightedPosition;
@@ -1425,10 +1425,11 @@ void  FQuadricSkeletalMeshReduction::AddSourceModelInfluences( const FSkeletalMe
 			for (int32 i = 0; i < SrcModelInfluences.Num(); ++i)
 			{
 				const SkeletalMeshImportData::FVertInfluence& VertInfluence = SrcModelInfluences[i];
-				const int32 VtxId = VertInfluence.VertIndex;  checkSlow(VtxId < NumImportVertex);
-				
-				
-				VtxToBoneIdWeightMap[VtxId].SetElement(VertInfluence.BoneIndex, VertInfluence.Weight);
+				const int32 VtxId = VertInfluence.VertIndex;
+				if (VtxId < NumImportVertex)
+				{
+					VtxToBoneIdWeightMap[VtxId].SetElement(VertInfluence.BoneIndex, VertInfluence.Weight);
+				}
 			}
 
 			// sort the bones and limit to MaxBonesPerVertex
@@ -1545,11 +1546,16 @@ void FQuadricSkeletalMeshReduction::ConvertToFSkeletalMeshLODModel( const int32
 		DummyMap[PointIdx] = PointIdx;
 	}
 
-	// Make sure we do not recalculate normals
+	// Make sure we do not recalculate normals or remove any degenerated data (threshold force to zero)
 
 	IMeshUtilities::MeshBuildOptions Options;
 	Options.bComputeNormals = false;
 	Options.bComputeTangents = false;
+	Options.bUseMikkTSpace = true; //Avoid builtin build by specifying true for mikkt space
+	Options.OverlappingThresholds.ThresholdPosition = 0.0f;
+	Options.OverlappingThresholds.ThresholdTangentNormal = 0.0f;
+	Options.OverlappingThresholds.ThresholdUV = 0.0f;
+	Options.bRemoveDegenerateTriangles = false;
 	IMeshUtilities& MeshUtilities = FModuleManager::Get().LoadModuleChecked<IMeshUtilities>("MeshUtilities");
 	
 	// Create skinning streams for NewModel.
@@ -1568,7 +1574,7 @@ void FQuadricSkeletalMeshReduction::ConvertToFSkeletalMeshLODModel( const int32
 	// Set texture coordinate count on the new model.
 	NewModel.NumTexCoords = SkeletalMeshData.TexCoordCount;
 
-	checkSlow(NewModel.MaxImportVertex + 1  == SkinnedMesh.NumVertices()); 
+
 
 	// Update the alternate weights
 	
@@ -1612,7 +1618,8 @@ bool FQuadricSkeletalMeshReduction::ReduceSkeletalLODModel( const FSkeletalMeshL
 	ConvertToFSkinnedSkeletalMesh(SrcModel, BoneMatrices, LODIndex, SkinnedSkeletalMesh);
 
 	int32 IterationNum = 0;
-
+	//We keep the original MaxNumVerts because if we iterate we want to still compare with the original request.
+	uint32 OriginalMaxNumVertsSetting = Settings.MaxNumOfVerts;
 	do 
 	{
 		if (bOptimizeMesh)
@@ -1655,8 +1662,7 @@ bool FQuadricSkeletalMeshReduction::ReduceSkeletalLODModel( const FSkeletalMeshL
 
 		// We may need to do additional simplification if the user specified a hard number limit for verts and
 		// the internal chunking during conversion split some verts.
-
-		if (bUseMaxVertexCriterion && OutSkeletalMeshLODModel.NumVertices > Settings.MaxNumOfVerts)
+		if (bUseMaxVertexCriterion && OutSkeletalMeshLODModel.NumVertices > OriginalMaxNumVertsSetting && OutSkeletalMeshLODModel.NumVertices > 6)
 		{
 			const bool bTerminatedOnVertCount = (Settings.TerminationCriterion == SMTC_AbsNumOfVerts) ||
 				                                (Settings.TerminationCriterion == SMTC_AbsTriangleOrVert && !(SkinnedSkeletalMesh.NumIndices() / 3 <= (int32)Settings.MaxNumOfTriangles));
@@ -1664,14 +1670,11 @@ bool FQuadricSkeletalMeshReduction::ReduceSkeletalLODModel( const FSkeletalMeshL
 			if (bTerminatedOnVertCount)
 			{
 				// Some verts were created by chunking - we need simplify more.
-				int32 ExcessVerts = (int32)OutSkeletalMeshLODModel.NumVertices - (int32)Settings.MaxNumOfVerts + IterationNum;
+				int32 ExcessVerts = (int32)(OutSkeletalMeshLODModel.NumVertices - OriginalMaxNumVertsSetting);
 				Settings.MaxNumOfVerts = FMath::Max((int32)Settings.MaxNumOfVerts - ExcessVerts, 6);
 
 				UE_LOG(LogSkeletalMeshReduction, Log, TEXT("Chunking to limit unique bones per section generated additional vertices - continuing simplification of LOD %d "), LODIndex);
-
-				ConvertToFSkinnedSkeletalMesh(OutSkeletalMeshLODModel, BoneMatrices, LODIndex, SkinnedSkeletalMesh);
-
-				
+				ConvertToFSkinnedSkeletalMesh(SrcModel, BoneMatrices, LODIndex, SkinnedSkeletalMesh);
 			}
 			else
 			{
@@ -1957,7 +1960,7 @@ void FQuadricSkeletalMeshReduction::ReduceSkeletalMesh(USkeletalMesh& SkeletalMe
 			// fix up chunks to remove the bones that set to be removed
 			for (int32 SectionIndex = 0; SectionIndex < NewModel->Sections.Num(); ++SectionIndex)
 			{
-				MeshBoneReductionInterface->FixUpSectionBoneMaps(NewModel->Sections[SectionIndex], BonesToRemove);
+				MeshBoneReductionInterface->FixUpSectionBoneMaps(NewModel->Sections[SectionIndex], BonesToRemove, NewModel->SkinWeightProfiles);
 			}
 		}
 
@@ -2030,7 +2033,7 @@ void FQuadricSkeletalMeshReduction::ReduceSkeletalMesh(USkeletalMesh& SkeletalMe
 			// fix up chunks to remove the bones that set to be removed
 			for (int32 SectionIndex = 0; SectionIndex < NewModel->Sections.Num(); ++SectionIndex)
 			{
-				MeshBoneReductionInterface->FixUpSectionBoneMaps(NewModel->Sections[SectionIndex], BonesToRemove);
+				MeshBoneReductionInterface->FixUpSectionBoneMaps(NewModel->Sections[SectionIndex], BonesToRemove, NewModel->SkinWeightProfiles);
 			}
 		}
 
diff --git a/Engine/Source/Developer/MeshBoneReduction/Private/MeshBoneReduction.cpp b/Engine/Source/Developer/MeshBoneReduction/Private/MeshBoneReduction.cpp
--- a/Engine/Source/Developer/MeshBoneReduction/Private/MeshBoneReduction.cpp
+++ b/Engine/Source/Developer/MeshBoneReduction/Private/MeshBoneReduction.cpp
@@ -129,7 +129,7 @@ public:
 		return ( OutBonesToReplace.Num() > 0 );
 	}
 
-	void FixUpSectionBoneMaps( FSkelMeshSection & Section, const TMap<FBoneIndexType, FBoneIndexType> &BonesToRepair ) override
+	void FixUpSectionBoneMaps( FSkelMeshSection & Section, const TMap<FBoneIndexType, FBoneIndexType> &BonesToRepair, TMap<FName, FImportedSkinWeightProfileData>& SkinWeightProfiles) override
 	{
 		// now you have list of bones, remove them from vertex influences
 		{
@@ -214,39 +214,55 @@ public:
 
 			if ( BoneMapRemapTable.Num() > 0 )
 			{
+				int32 BaseVertexIndex = Section.BaseVertexIndex;
 				// fix up soft verts
 				for (int32 VertIndex=0; VertIndex < Section.SoftVertices.Num(); ++VertIndex)
 				{
 					FSoftSkinVertex & Vert = Section.SoftVertices[VertIndex];
-					bool ShouldRenormalize = false;
 
-					for(int32 InfluenceIndex = 0;InfluenceIndex < MAX_TOTAL_INFLUENCES;InfluenceIndex++)
+					auto RemapBoneInfluenceVertexIndex = [&BoneMapRemapTable](uint8 InfluenceBones[MAX_TOTAL_INFLUENCES], uint8 InfluenceWeights[MAX_TOTAL_INFLUENCES])
 					{
-						uint8 *RemappedBone = BoneMapRemapTable.Find(Vert.InfluenceBones[InfluenceIndex]);
-						if (RemappedBone)
+						bool ShouldRenormalize = false;
+
+						for (int32 InfluenceIndex = 0; InfluenceIndex < MAX_TOTAL_INFLUENCES; InfluenceIndex++)
 						{
-							Vert.InfluenceBones[InfluenceIndex] = *RemappedBone;
-							ShouldRenormalize = true;
+							uint8 *RemappedBone = BoneMapRemapTable.Find(InfluenceBones[InfluenceIndex]);
+							if (RemappedBone)
+							{
+								InfluenceBones[InfluenceIndex] = *RemappedBone;
+								ShouldRenormalize = true;
+							}
 						}
-					}
 
-					if (ShouldRenormalize)
-					{
-						// should see if same bone exists
-						for(int32 InfluenceIndex = 0;InfluenceIndex < MAX_TOTAL_INFLUENCES;InfluenceIndex++)
+						if (ShouldRenormalize)
 						{
-							for(int32 InfluenceIndex2 = InfluenceIndex+1;InfluenceIndex2 < MAX_TOTAL_INFLUENCES;InfluenceIndex2++)
+							// should see if same bone exists
+							for (int32 InfluenceIndex = 0; InfluenceIndex < MAX_TOTAL_INFLUENCES; InfluenceIndex++)
 							{
-								// cannot be 0 because we don't allow removing root
-								if (Vert.InfluenceBones[InfluenceIndex] != 0 && Vert.InfluenceBones[InfluenceIndex] == Vert.InfluenceBones[InfluenceIndex2])
+								for (int32 InfluenceIndex2 = InfluenceIndex + 1; InfluenceIndex2 < MAX_TOTAL_INFLUENCES; InfluenceIndex2++)
 								{
-									Vert.InfluenceWeights[InfluenceIndex] += Vert.InfluenceWeights[InfluenceIndex2];
-									// reset
-									Vert.InfluenceBones[InfluenceIndex2] = 0;
-									Vert.InfluenceWeights[InfluenceIndex2] = 0;
+									// cannot be 0 because we don't allow removing root
+									if (InfluenceBones[InfluenceIndex] != 0 && InfluenceBones[InfluenceIndex] == InfluenceBones[InfluenceIndex2])
+									{
+										InfluenceWeights[InfluenceIndex] += InfluenceWeights[InfluenceIndex2];
+										// reset
+										InfluenceBones[InfluenceIndex2] = 0;
+										InfluenceWeights[InfluenceIndex2] = 0;
+									}
 								}
 							}
 						}
+					};
+
+					RemapBoneInfluenceVertexIndex(Vert.InfluenceBones, Vert.InfluenceWeights);
+
+					int32 RealVertexIndex = BaseVertexIndex + VertIndex;
+					//Remap the alternate weights
+					for (auto Kvp : SkinWeightProfiles)
+					{
+						FImportedSkinWeightProfileData& SkinWeightProfile = SkinWeightProfiles.FindChecked(Kvp.Key);
+						check(SkinWeightProfile.SkinWeights.IsValidIndex(RealVertexIndex));
+						RemapBoneInfluenceVertexIndex(SkinWeightProfile.SkinWeights[RealVertexIndex].InfluenceBones, SkinWeightProfile.SkinWeights[RealVertexIndex].InfluenceWeights);
 					}
 				}
 			}
@@ -443,7 +459,7 @@ public:
 						Vertex.TangentZ.W = WComponent;
 					}
 				}
-				FixUpSectionBoneMaps(Section, BonesToRemove);
+				FixUpSectionBoneMaps(Section, BonesToRemove, NewModel->SkinWeightProfiles);
 			});
 
 			// fix up RequiredBones/ActiveBoneIndices
diff --git a/Engine/Source/Developer/MeshBoneReduction/Public/MeshBoneReduction.h b/Engine/Source/Developer/MeshBoneReduction/Public/MeshBoneReduction.h
--- a/Engine/Source/Developer/MeshBoneReduction/Public/MeshBoneReduction.h
+++ b/Engine/Source/Developer/MeshBoneReduction/Public/MeshBoneReduction.h
@@ -24,7 +24,7 @@ public:
 	 *
 	 * @return	true if success 
 	 */
-	virtual void FixUpSectionBoneMaps( FSkelMeshSection & Section, const TMap<FBoneIndexType, FBoneIndexType> &BonesToRepair ) = 0;
+	virtual void FixUpSectionBoneMaps( FSkelMeshSection & Section, const TMap<FBoneIndexType, FBoneIndexType> &BonesToRepair, TMap<FName, FImportedSkinWeightProfileData>& SkinWeightProfiles) = 0;
 	/**
 	 * Get Bones To Remove from the Desired LOD
 	 * List of bones to remove should contains <bone index to remove, bone index to replace to >
diff --git a/Engine/Source/Editor/UnrealEd/Private/LODUtilities.cpp b/Engine/Source/Editor/UnrealEd/Private/LODUtilities.cpp
--- a/Engine/Source/Editor/UnrealEd/Private/LODUtilities.cpp
+++ b/Engine/Source/Editor/UnrealEd/Private/LODUtilities.cpp
@@ -1802,7 +1802,7 @@ bool FLODUtilities::ImportAlternateSkinWeight(USkeletalMesh* SkeletalMesh, FStri
 	check(SkeletalMesh->GetLODInfo(TargetLODIndex));
 	FSkeletalMeshLODInfo* LODInfo = SkeletalMesh->GetLODInfo(TargetLODIndex);
 	
-	if (LODInfo && LODInfo->bHasBeenSimplified)
+	if (LODInfo && LODInfo->bHasBeenSimplified && LODInfo->ReductionSettings.BaseLOD != TargetLODIndex)
 	{
 		//We cannot remove alternate skin weights profile for a generated LOD
 		UE_LOG(LogLODUtilities, Error, TEXT("Cannot import Skin Weight Profile for a generated LOD."));
