Fix C++ standard conformance issues with latest version of Visual Studio running with /permissive-.

diff --git a/Engine/Plugins/Experimental/SkeletalReduction/Source/Private/SkeletalSimplifierLinearAlgebra.h b/Engine/Plugins/Experimental/SkeletalReduction/Source/Private/SkeletalSimplifierLinearAlgebra.h
--- a/Engine/Plugins/Experimental/SkeletalReduction/Source/Private/SkeletalSimplifierLinearAlgebra.h
+++ b/Engine/Plugins/Experimental/SkeletalReduction/Source/Private/SkeletalSimplifierLinearAlgebra.h
@@ -248,11 +248,12 @@ namespace SkeletalSimplifier
 			typedef ArrayBase<6, SymmetricMatrix>            MyBase;
 			typedef typename MyBase::ScalarType              ScalarType;
 
+			using MyBase::Data;
 
 			SymmetricMatrix() 
 			{
 				// Zero
-				Reset();
+				this->Reset();
 			}
 
 			/**
diff --git a/Engine/Plugins/Experimental/SkeletalReduction/Source/Private/SkeletalSimplifierQuadrics.h b/Engine/Plugins/Experimental/SkeletalReduction/Source/Private/SkeletalSimplifierQuadrics.h
--- a/Engine/Plugins/Experimental/SkeletalReduction/Source/Private/SkeletalSimplifierQuadrics.h
+++ b/Engine/Plugins/Experimental/SkeletalReduction/Source/Private/SkeletalSimplifierQuadrics.h
@@ -168,6 +168,18 @@ namespace SkeletalSimplifier
 			typedef typename MyBase::D1VectorType                              D1VectorType;
 			typedef typename MyBase::D2VectorType                              D2VectorType;
 
+			using MyBase::CMatrix;
+			using MyBase::B1Matrix;
+			using MyBase::B2Matrix;
+			using MyBase::Gamma;
+			using MyBase::D0Vector;
+			using MyBase::D1Vector;
+			using MyBase::D2Vector;
+			using MyBase::CScalar;
+			using MyBase::VolGradConstraint;
+			using MyBase::VolDistConstraint;
+			using MyBase::UVBBox;
+
 			/**
 			* @param TriVert0  - The three verts that define the triangle face.
 			* @param TriVert1  -
@@ -282,6 +294,18 @@ namespace SkeletalSimplifier
 			typedef TQuadricBase<BaseAttrContainerType, AttrContainerType>             MyBase;
 			typedef FaceQuadricType                                                    FFaceQuadric;
 
+			using MyBase::CMatrix;
+			using MyBase::B1Matrix;
+			using MyBase::B2Matrix;
+			using MyBase::Gamma;
+			using MyBase::D0Vector;
+			using MyBase::D1Vector;
+			using MyBase::D2Vector;
+			using MyBase::CScalar;
+			using MyBase::VolGradConstraint;
+			using MyBase::VolDistConstraint;
+			using MyBase::UVBBox;
+
 			TQuadricOptimizer() :
 				MyBase()
 			{}
