PR #5322: Fix unqualified lookup in dependent base in TDenseVecD (Contributed by [REDACTED_PERSON])

diff --git a/Engine/Plugins/Experimental/SkeletalReduction/Source/Private/SkeletalSimplifierLinearAlgebra.h b/Engine/Plugins/Experimental/SkeletalReduction/Source/Private/SkeletalSimplifierLinearAlgebra.h
--- a/Engine/Plugins/Experimental/SkeletalReduction/Source/Private/SkeletalSimplifierLinearAlgebra.h
+++ b/Engine/Plugins/Experimental/SkeletalReduction/Source/Private/SkeletalSimplifierLinearAlgebra.h
@@ -248,12 +248,10 @@ namespace SkeletalSimplifier
 			typedef ArrayBase<6, SymmetricMatrix>            MyBase;
 			typedef typename MyBase::ScalarType              ScalarType;
 
-			using MyBase::Data;
-
 			SymmetricMatrix() 
 			{
 				// Zero
-				this->Reset();
+				Reset();
 			}
 
 			/**
@@ -966,6 +964,9 @@ namespace SkeletalSimplifier
 			typedef ArrayBase<SIZE, TDenseVecD>           MyBase;
 			typedef typename MyBase::ScalarType           ScalarType;
 
+			using MyBase::Data;
+			using MyBase::Reset;
+
 			enum { Size = SIZE };
 
 			TDenseVecD()
