Skeletal Mesh Reduction:
Replace check with warning when the bone weights don't sum up.

diff --git a/Engine/Plugins/Experimental/SkeletalReduction/Source/Private/SkeletalMeshReductionPlugin.cpp b/Engine/Plugins/Experimental/SkeletalReduction/Source/Private/SkeletalMeshReductionPlugin.cpp
--- a/Engine/Plugins/Experimental/SkeletalReduction/Source/Private/SkeletalMeshReductionPlugin.cpp
+++ b/Engine/Plugins/Experimental/SkeletalReduction/Source/Private/SkeletalMeshReductionPlugin.cpp
@@ -448,16 +448,45 @@ void FQuadricSkeletalMeshReduction::ConvertToFSkinnedSkeletalMesh( const FSkelet
 		Vertex.Position   = WeightedPosition;
 	};
 
-	auto CreatSkinningMatrix = [&BoneMatrices](const FSoftSkinVertex& Vertex, const FSkelMeshSection& Section)->FMatrix
+	auto CreateSkinningMatrix = [&BoneMatrices](const FSoftSkinVertex& Vertex, const FSkelMeshSection& Section, bool& bValidBoneWeights)->FMatrix
 	{
-		FMatrix BlendedMatrix(ForceInitToZero);
-		int32 TotalInfluence = 0;
-		int32 ValidInfluenceCount = 0;
+		// Compute the inverse of the total bone influence for this vertex.
+		
+		float InvTotalInfluence = 1.f / 255.f;   // expected default - anything else could indicate a problem with the asset.
+		{
+			int32 TotalInfluence = 0;
+
+			for (int32 i = 0; i < MAX_TOTAL_INFLUENCES; ++i)
+			{
+				const uint8 BoneInfluence = Vertex.InfluenceWeights[i];
+				TotalInfluence += BoneInfluence;
+			}
+
+			if (TotalInfluence != 255) // 255 is the expected value.  This logic just allows for graceful failure.
+			{
+				// Not expected value - record that.
+				bValidBoneWeights = false;
+
+				if (TotalInfluence == 0)
+				{
+					InvTotalInfluence = 0.f;
+				}
+				else
+				{
+					InvTotalInfluence = 1.f / float(TotalInfluence);
+				}
+			}
+		}
 
-		const TArray<uint16>& BoneMap = Section.BoneMap;
 
 		// Build the blended matrix 
 
+		FMatrix BlendedMatrix(ForceInitToZero);
+
+		int32 ValidInfluenceCount = 0;
+		
+		const TArray<uint16>& BoneMap = Section.BoneMap;
+
 		for (int32 i = 0; i < MAX_TOTAL_INFLUENCES; ++i)
 		{
 			const uint16 BoneIndex    = Vertex.InfluenceBones[i];
@@ -465,13 +494,11 @@ void FQuadricSkeletalMeshReduction::ConvertToFSkinnedSkeletalMesh( const FSkelet
 
 			// Accumulate the bone influence for this vert into the BlendedMatrix
 
-			TotalInfluence += BoneInfluence;
-
 			if (BoneInfluence > 0)
 			{
 				check(BoneIndex < BoneMap.Num());
 				const uint16 SectionBoneId = BoneMap[BoneIndex]; // Third-party tool uses an additional indirection bode table here 
-				const float  BoneWeight = BoneInfluence / 255.0f;  // convert to [0,1] float
+				const float  BoneWeight = BoneInfluence * InvTotalInfluence;  // convert to [0,1] float
 
 				if (BoneMatrices.IsValidIndex(SectionBoneId))
 				{
@@ -483,12 +510,13 @@ void FQuadricSkeletalMeshReduction::ConvertToFSkinnedSkeletalMesh( const FSkelet
 		}
 
 		// default identity matrix for the special case of the vertex having no valid transforms..
+
 		if (ValidInfluenceCount == 0)
 		{
 			BlendedMatrix = FMatrix::Identity;
 		}
 
-		check(TotalInfluence == 255);
+		
 
 		return BlendedMatrix;
 	};
@@ -548,18 +576,24 @@ void FQuadricSkeletalMeshReduction::ConvertToFSkinnedSkeletalMesh( const FSkelet
 		const FSectionRange VertexRange  = SectionRangeArray[SectionIndex];
 		
 		// Loop over the vertices in this section.
-
+		bool bHasValidBoneWeights = true;
 		for (int32 VertexIndex = VertexRange.Begin; VertexIndex < VertexRange.End; ++VertexIndex)
 		{
 			FSoftSkinVertex& SkinVertex = SoftSkinVertices[VertexIndex];
 
 			// Use the bone weights for this vertex to create a blended matrix 
-			const FMatrix BlendedMatrix = CreatSkinningMatrix(SkinVertex, Section);
+			const FMatrix BlendedMatrix = CreateSkinningMatrix(SkinVertex, Section, bHasValidBoneWeights);
 
 			// Update this Skin Vertex to the correct location, normal, etc.
 			ApplySkinning(BlendedMatrix, SkinVertex);
 
 		}
+
+		// Report any error with invalid bone weights
+		if (!bHasValidBoneWeights && !SkipSection(SectionIndex))
+		{
+			UE_LOG(LogSkeletalMeshReduction, Warning, TEXT("Building LOD %d - Encountered questionable vertex weights in source."), LODIndex);
+		}
 	}
 
 	// -- Make the index buffer; skipping the "SkipSections"
