Fix a crash when doing a LOD reduction base on the same LOD and there is no import data.

diff --git a/Engine/Plugins/Experimental/SkeletalReduction/Source/Private/SkeletalMeshReductionPlugin.cpp b/Engine/Plugins/Experimental/SkeletalReduction/Source/Private/SkeletalMeshReductionPlugin.cpp
--- a/Engine/Plugins/Experimental/SkeletalReduction/Source/Private/SkeletalMeshReductionPlugin.cpp
+++ b/Engine/Plugins/Experimental/SkeletalReduction/Source/Private/SkeletalMeshReductionPlugin.cpp
@@ -1309,7 +1309,12 @@ void FQuadricSkeletalMeshReduction::ReduceSkeletalMesh(USkeletalMesh& SkeletalMe
 	// only allow to set BaseLOD if the LOD is less than this
 	if (Settings.BaseLOD > 0)
 	{
-		if (Settings.BaseLOD <= LODIndex && SkeletalMeshResource.LODModels.IsValidIndex(Settings.BaseLOD))
+		if (Settings.BaseLOD == LODIndex && (!SkelResource->OriginalReductionSourceMeshData.IsValidIndex(BaseLOD) || SkelResource->OriginalReductionSourceMeshData[BaseLOD]->IsEmpty()))
+		{
+			//Cannot reduce ourself if we are not imported
+			UE_LOG(LogSkeletalMeshReduction, Warning, TEXT("Building LOD %d - Cannot generate LOD with himself if the LOD do not have imported Data. Using Base LOD 0 instead"), LODIndex);
+		}
+		else if (Settings.BaseLOD <= LODIndex && SkeletalMeshResource.LODModels.IsValidIndex(Settings.BaseLOD))
 		{
 			BaseLOD = Settings.BaseLOD;
 			SrcModel = &SkeletalMeshResource.LODModels[BaseLOD];
@@ -1317,7 +1322,7 @@ void FQuadricSkeletalMeshReduction::ReduceSkeletalMesh(USkeletalMesh& SkeletalMe
 		else
 		{
 			// warn users
-			UE_LOG(LogSkeletalMeshReduction, Warning, TEXT("Building LOD %d - Invalid Base LOD entered. Using Base LOD 0"), LODIndex);
+			UE_LOG(LogSkeletalMeshReduction, Warning, TEXT("Building LOD %d - Invalid Base LOD entered. Using Base LOD 0 instead"), LODIndex);
 		}
 	}
 	//Reducing Base LOD, we need to use the temporary data so it can be iterative
