Fix the skeletalmesh inline LOD reduce

diff --git a/Engine/Plugins/Experimental/SkeletalReduction/Source/Private/SkeletalMeshReductionPlugin.cpp b/Engine/Plugins/Experimental/SkeletalReduction/Source/Private/SkeletalMeshReductionPlugin.cpp
--- a/Engine/Plugins/Experimental/SkeletalReduction/Source/Private/SkeletalMeshReductionPlugin.cpp
+++ b/Engine/Plugins/Experimental/SkeletalReduction/Source/Private/SkeletalMeshReductionPlugin.cpp
@@ -1309,7 +1309,7 @@ void FQuadricSkeletalMeshReduction::ReduceSkeletalMesh(USkeletalMesh& SkeletalMe
 	// only allow to set BaseLOD if the LOD is less than this
 	if (Settings.BaseLOD > 0)
 	{
-		if (Settings.BaseLOD == LODIndex && (!SkelResource->OriginalReductionSourceMeshData.IsValidIndex(BaseLOD) || SkelResource->OriginalReductionSourceMeshData[BaseLOD]->IsEmpty()))
+		if (Settings.BaseLOD == LODIndex && (!SkelResource->OriginalReductionSourceMeshData.IsValidIndex(Settings.BaseLOD) || SkelResource->OriginalReductionSourceMeshData[Settings.BaseLOD]->IsEmpty()))
 		{
 			//Cannot reduce ourself if we are not imported
 			UE_LOG(LogSkeletalMeshReduction, Warning, TEXT("Building LOD %d - Cannot generate LOD with himself if the LOD do not have imported Data. Using Base LOD 0 instead"), LODIndex);
