Skeletal Mesh Reduction.  Aditional simplification when a vert target has been defeated by chunking.

diff --git a/Engine/Plugins/Experimental/SkeletalReduction/Source/Private/SkeletalMeshReductionPlugin.cpp b/Engine/Plugins/Experimental/SkeletalReduction/Source/Private/SkeletalMeshReductionPlugin.cpp
--- a/Engine/Plugins/Experimental/SkeletalReduction/Source/Private/SkeletalMeshReductionPlugin.cpp
+++ b/Engine/Plugins/Experimental/SkeletalReduction/Source/Private/SkeletalMeshReductionPlugin.cpp
@@ -157,7 +157,7 @@ private:
 		                         FSkeletalMeshLODModel& OutLODModel,
 		                         const FBoxSphereBounds& Bounds,
 		                         const FReferenceSkeleton& RefSkeleton,
-		                         const FSkeletalMeshOptimizationSettings& Settings,
+		                         FSkeletalMeshOptimizationSettings Settings,
 								 const FImportantBones& ImportantBones,
 		                         const TArray<FMatrix>& BoneMatrices,
 		                         const int32 LODIndex ) const;
@@ -1259,63 +1259,113 @@ bool FQuadricSkeletalMeshReduction::ReduceSkeletalLODModel( const FSkeletalMeshL
 	                                                        FSkeletalMeshLODModel& OutSkeletalMeshLODModel,
 	                                                        const FBoxSphereBounds& Bounds,
 	                                                        const FReferenceSkeleton& RefSkeleton,
-	                                                        const FSkeletalMeshOptimizationSettings& Settings,
+	                                                        FSkeletalMeshOptimizationSettings Settings,
 															const FImportantBones& ImportantBones,
 	                                                        const TArray<FMatrix>& BoneMatrices,
 	                                                        const int32 LODIndex
                                                            ) const
 {
+
+	const int32 SrcNumVerts = SrcModel.NumVertices;
+	
 	// Parameters for Simplification etc
-	const bool bUseVertexCriterion   = (Settings.TerminationCriterion != SMTC_NumOfTriangles && Settings.NumOfVertPercentage < 1.f);
-	const bool bUseTriangleCriterion = (Settings.TerminationCriterion != SMTC_NumOfVerts && Settings.NumOfTrianglesPercentage < 1.f);
-	const bool bProcessGeometry      = (bUseTriangleCriterion || bUseVertexCriterion);
-	const bool bProcessBones         = (Settings.MaxBonesPerVertex < MAX_TOTAL_INFLUENCES);
-	const bool bOptimizeMesh         = (bProcessGeometry || bProcessBones);
+	const bool bUseVertexPercentCriterion   = ((Settings.TerminationCriterion == SMTC_NumOfVerts     || Settings.TerminationCriterion == SMTC_TriangleOrVert) && Settings.NumOfVertPercentage < 1.f) ;
+	const bool bUseTrianglePercentCriterion = ((Settings.TerminationCriterion == SMTC_NumOfTriangles || Settings.TerminationCriterion == SMTC_TriangleOrVert) && Settings.NumOfTrianglesPercentage < 1.f);
 
+	const bool bUseMaxVertexCriterion   = ((Settings.TerminationCriterion == SMTC_AbsNumOfVerts || Settings.TerminationCriterion == SMTC_AbsTriangleOrVert) && SrcNumVerts);
+	const bool bUseMaxTriangleCriterion = ((Settings.TerminationCriterion == SMTC_AbsNumOfTriangles || Settings.TerminationCriterion == SMTC_AbsTriangleOrVert) && Settings.MaxNumOfTriangles < INT32_MAX);
+
+	const bool bProcessGeometry      = (bUseTrianglePercentCriterion || bUseVertexPercentCriterion || bUseMaxTriangleCriterion || bUseMaxVertexCriterion);
+	const bool bProcessBones         = (Settings.MaxBonesPerVertex < MAX_TOTAL_INFLUENCES);
+	
+	bool bOptimizeMesh         = (bProcessGeometry || bProcessBones);
 
+	if (bOptimizeMesh)
+	{
+		UE_LOG(LogSkeletalMeshReduction, Log, TEXT("Reducing skeletal mesh for LOD %d "), LODIndex);
+	}
 	
 	// Generate a single skinned mesh form the SrcModel.  This mesh has per-vertex tangent space.
 
 	SkeletalSimplifier::FSkinnedSkeletalMesh SkinnedSkeletalMesh;
 	ConvertToFSkinnedSkeletalMesh(SrcModel, BoneMatrices, LODIndex, SkinnedSkeletalMesh);
 
-	if (bOptimizeMesh)
+	int32 IterationNum = 0;
+
+	do 
 	{
-		if (ImportantBones.Ids.Num() > 0)
+		if (bOptimizeMesh)
 		{
-			// Add specialized weights for verts associated with "important" bones.
-			UpdateSpecializedVertWeights(ImportantBones, SkinnedSkeletalMesh);
-		}
+			if (ImportantBones.Ids.Num() > 0)
+			{
+				// Add specialized weights for verts associated with "important" bones.
+				UpdateSpecializedVertWeights(ImportantBones, SkinnedSkeletalMesh);
+			}
 
-		// Capture the UV bounds from the source mesh.
+			// Capture the UV bounds from the source mesh.
 
-		FVector2D  UVBounds[2 * SkeletalSimplifier::MeshVertType::BasicAttrContainerType::NumUVs];
-		ComputeUVBounds(SkinnedSkeletalMesh, UVBounds);
+			FVector2D  UVBounds[2 * SkeletalSimplifier::MeshVertType::BasicAttrContainerType::NumUVs];
+			ComputeUVBounds(SkinnedSkeletalMesh, UVBounds);
 
-		{
-			// Use the bone-aware simplifier
+			{
+				// Use the bone-aware simplifier
+
+				SimplifyMesh(Settings, Bounds, SkinnedSkeletalMesh);
+			}
+
+			// Clamp the UVs of the simplified mesh to match the source mesh.
 
-			SimplifyMesh(Settings, Bounds, SkinnedSkeletalMesh);
+			ClampUVBounds(UVBounds, SkinnedSkeletalMesh);
+
+
+			// Reduce the number of bones per-vert
+
+			const int32 MaxBonesPerVert = FMath::Clamp(Settings.MaxBonesPerVertex, 0, MAX_TOTAL_INFLUENCES);
+
+			if (MaxBonesPerVert < MAX_TOTAL_INFLUENCES)
+			{
+				TrimBonesPerVert(SkinnedSkeletalMesh, MaxBonesPerVert);
+			}
 		}
 
-		// Clamp the UVs of the simplified mesh to match the source mesh.
+		// Convert to SkeletalMeshLODModel. 
+
+		ConvertToFSkeletalMeshLODModel(SkinnedSkeletalMesh, RefSkeleton, OutSkeletalMeshLODModel);
 
-		ClampUVBounds(UVBounds, SkinnedSkeletalMesh);
+		// We may need to do additional simplification if the user specified a hard number limit for verts and
+		// the internal chunking during conversion split some verts.
+
+		if (bUseMaxVertexCriterion && OutSkeletalMeshLODModel.NumVertices > Settings.MaxNumOfVerts)
+		{
+			const bool bTerminatedOnVertCount = (Settings.TerminationCriterion == SMTC_AbsNumOfVerts) ||
+				                                (Settings.TerminationCriterion == SMTC_AbsTriangleOrVert && !(SkinnedSkeletalMesh.NumIndices() / 3 <= (int32)Settings.MaxNumOfTriangles));
+			  
+			if (bTerminatedOnVertCount)
+			{
+				// Some verts were created by chunking - we need simplify more.
+				int32 ExcessVerts = (int32)OutSkeletalMeshLODModel.NumVertices - (int32)Settings.MaxNumOfVerts + IterationNum;
+				Settings.MaxNumOfVerts = FMath::Max((int32)Settings.MaxNumOfVerts - ExcessVerts, 6);
 
+				UE_LOG(LogSkeletalMeshReduction, Log, TEXT("Chunking to limit unique bones per section generated additional vertices - continuing simplification of LOD %d "), LODIndex);
 
-		// Reduce the number of bones per-vert
+				ConvertToFSkinnedSkeletalMesh(OutSkeletalMeshLODModel, BoneMatrices, LODIndex, SkinnedSkeletalMesh);
 
-		const int32 MaxBonesPerVert = FMath::Clamp(Settings.MaxBonesPerVertex, 0, MAX_TOTAL_INFLUENCES);
+				
+			}
+			else
+			{
+				bOptimizeMesh = false;
+			}
 
-		if (MaxBonesPerVert < MAX_TOTAL_INFLUENCES)
+			IterationNum++;
+		}
+		else
 		{
-			TrimBonesPerVert(SkinnedSkeletalMesh, MaxBonesPerVert);
+			bOptimizeMesh = false;
 		}
-	}
+	} 
+	while (bOptimizeMesh && IterationNum < 5);
 
-	// Convert to SkeletalMeshLODModel. 
-		
-	ConvertToFSkeletalMeshLODModel(SkinnedSkeletalMesh, RefSkeleton, OutSkeletalMeshLODModel);
 
 	bool bReturnValue =  (OutSkeletalMeshLODModel.NumVertices > 0);
 	
