Moved skeletal mesh reduction plugin to early access

diff --git a/Engine/Plugins/Experimental/SkeletalReduction/SkeletalReduction.uplugin b/Engine/Plugins/Experimental/SkeletalReduction/SkeletalReduction.uplugin
--- a/Engine/Plugins/Experimental/SkeletalReduction/SkeletalReduction.uplugin
+++ b/Engine/Plugins/Experimental/SkeletalReduction/SkeletalReduction.uplugin
@@ -1,8 +1,8 @@
 {
   "FileVersion" : 3,
   "Version" : 1,
-  "VersionName" : "0.1",
-  "FriendlyName" : "Skeletal Mesh Simplifier (Experimental)",
+  "VersionName" : "1.0",
+  "FriendlyName" : "Skeletal Mesh Simplifier (Early Access)",
   "Description" : "A plugin to generate LOD for deforming meshes.",
   "Category" : "Editor",
   "CreatedBy" : "Epic Games, Inc.",
@@ -12,7 +12,7 @@
   "SupportURL" : "",
   "EnabledByDefault" : true,
   "CanContainContent" : false,
-  "IsBetaVersion" : true,
+  "IsBetaVersion" : false,
   "Installed" : false,
   "Modules" :
 	[
