Skeletal Mesh Reduction: Make the simplifier more robust to NaNs in the tangent space.

diff --git a/Engine/Plugins/Experimental/SkeletalReduction/Source/Private/SkeletalMeshReductionPlugin.cpp b/Engine/Plugins/Experimental/SkeletalReduction/Source/Private/SkeletalMeshReductionPlugin.cpp
--- a/Engine/Plugins/Experimental/SkeletalReduction/Source/Private/SkeletalMeshReductionPlugin.cpp
+++ b/Engine/Plugins/Experimental/SkeletalReduction/Source/Private/SkeletalMeshReductionPlugin.cpp
@@ -461,21 +461,34 @@ void FQuadricSkeletalMeshReduction::ConvertToFSkinnedSkeletalMesh( const FSkelet
 
 
 
-	auto ApplySkinning = [](const FMatrix& XForm, FSoftSkinVertex& Vertex)->void
+	auto ApplySkinning = [](const FMatrix& XForm, FSoftSkinVertex& Vertex)->bool
 	{
+		// Some imported models will have garbage tangent space
+		bool bHasBadNTB =  ( Vertex.TangentX.ContainsNaN() || Vertex.TangentY.ContainsNaN() || Vertex.TangentZ.ContainsNaN() );
+
 		// transform position
 		FVector WeightedPosition = XForm.TransformPosition(Vertex.Position);
-		FVector WeightedTangentX = XForm.TransformVector(Vertex.TangentX);
-		FVector WeightedTangentY = XForm.TransformVector(Vertex.TangentY);
-		FVector WeightedTangentZ = XForm.TransformVector(Vertex.TangentZ);
+
+		FVector WeightedTangentX(1.f, 0.f, 0.f);
+		FVector WeightedTangentY(0.f, 1.f, 0.f);
+		FVector WeightedTangentZ(0.f ,0.f, 1.f);
+
+		if (!bHasBadNTB)
+		{
+			WeightedTangentX = XForm.TransformVector(Vertex.TangentX);
+			WeightedTangentY = XForm.TransformVector(Vertex.TangentY);
+			WeightedTangentZ = XForm.TransformVector(Vertex.TangentZ); // how does this work since TangentZ is a FVector4?
+		}
 
 		
 		Vertex.TangentX   = WeightedTangentX.GetSafeNormal();
 		Vertex.TangentY   = WeightedTangentY.GetSafeNormal();
-		float WComponent  = Vertex.TangentZ.W;
+		float WComponent  = (bHasBadNTB) ? 1.f : Vertex.TangentZ.W;             
 		Vertex.TangentZ   = WeightedTangentZ.GetSafeNormal();
 		Vertex.TangentZ.W = WComponent;
 		Vertex.Position   = WeightedPosition;
+
+		return bHasBadNTB;
 	};
 
 	auto CreateSkinningMatrix = [&BoneMatrices](const FSoftSkinVertex& Vertex, const FSkelMeshSection& Section, bool& bValidBoneWeights)->FMatrix
@@ -597,9 +610,8 @@ void FQuadricSkeletalMeshReduction::ConvertToFSkinnedSkeletalMesh( const FSkelet
 	check(TexCoordCount <= MAX_TEXCOORDS);
 
 	
-	
 	// Update the verts to the skinned location.
-
+	int32 NumBadNTBSpace = 0;
 	for (int32 SectionIndex = 0; SectionIndex < SectionCount; ++SectionIndex)
 	{
 		const FSkelMeshSection& Section  = SrcLODModel.Sections[SectionIndex];
@@ -614,9 +626,15 @@ void FQuadricSkeletalMeshReduction::ConvertToFSkinnedSkeletalMesh( const FSkelet
 			// Use the bone weights for this vertex to create a blended matrix 
 			const FMatrix BlendedMatrix = CreateSkinningMatrix(SkinVertex, Section, bHasValidBoneWeights);
 
+
 			// Update this Skin Vertex to the correct location, normal, etc.
-			ApplySkinning(BlendedMatrix, SkinVertex);
+			// also replace NaN tangent spaces with default tangent before skinning
+			bool bHasBadNTB = ApplySkinning(BlendedMatrix, SkinVertex);
 
+			if (bHasBadNTB)
+			{
+				NumBadNTBSpace++;
+		}
 		}
 
 		// Report any error with invalid bone weights
@@ -626,6 +644,11 @@ void FQuadricSkeletalMeshReduction::ConvertToFSkinnedSkeletalMesh( const FSkelet
 		}
 	}
 
+	if (NumBadNTBSpace > 0)
+	{
+		UE_LOG(LogSkeletalMeshReduction, Warning, TEXT("There were NaNs in the Tangent Space of %d source model vertices."), NumBadNTBSpace);
+	}
+
 	// -- Make the index buffer; skipping the "SkipSections"
 
 	// How many triangles?
@@ -1428,8 +1451,8 @@ void  FQuadricSkeletalMeshReduction::AddSourceModelInfluences( const FSkeletalMe
 				const int32 VtxId = VertInfluence.VertIndex;
 				if (VtxId < NumImportVertex)
 				{
-					VtxToBoneIdWeightMap[VtxId].SetElement(VertInfluence.BoneIndex, VertInfluence.Weight);
-				}
+				VtxToBoneIdWeightMap[VtxId].SetElement(VertInfluence.BoneIndex, VertInfluence.Weight);
+			}
 			}
 
 			// sort the bones and limit to MaxBonesPerVertex
