Skeletal simplifier : Make more robust to problems in tangent space.   Fixing Robomerge problem

diff --git a/Engine/Plugins/Experimental/SkeletalReduction/Source/Private/SkeletalMeshReductionPlugin.cpp b/Engine/Plugins/Experimental/SkeletalReduction/Source/Private/SkeletalMeshReductionPlugin.cpp
--- a/Engine/Plugins/Experimental/SkeletalReduction/Source/Private/SkeletalMeshReductionPlugin.cpp
+++ b/Engine/Plugins/Experimental/SkeletalReduction/Source/Private/SkeletalMeshReductionPlugin.cpp
@@ -469,17 +469,17 @@ void FQuadricSkeletalMeshReduction::ConvertToFSkinnedSkeletalMesh( const FSkelet
 		// transform position
 		FVector WeightedPosition = XForm.TransformPosition(Vertex.Position);
 
+		// transform tangent space
 		FVector WeightedTangentX(1.f, 0.f, 0.f);
 		FVector WeightedTangentY(0.f, 1.f, 0.f);
-		FVector WeightedTangentZ(0.f ,0.f, 1.f);
+		FVector WeightedTangentZ(0.f, 0.f, 1.f);
 
 		if (!bHasBadNTB)
 		{
 			WeightedTangentX = XForm.TransformVector(Vertex.TangentX);
 			WeightedTangentY = XForm.TransformVector(Vertex.TangentY);
-			WeightedTangentZ = XForm.TransformVector(Vertex.TangentZ); // how does this work since TangentZ is a FVector4?
+			WeightedTangentZ = XForm.TransformVector(Vertex.TangentZ);
 		}
-
 		
 		Vertex.TangentX   = WeightedTangentX.GetSafeNormal();
 		Vertex.TangentY   = WeightedTangentY.GetSafeNormal();
@@ -1451,8 +1451,8 @@ void  FQuadricSkeletalMeshReduction::AddSourceModelInfluences( const FSkeletalMe
 				const int32 VtxId = VertInfluence.VertIndex;
 				if (VtxId < NumImportVertex)
 				{
-				VtxToBoneIdWeightMap[VtxId].SetElement(VertInfluence.BoneIndex, VertInfluence.Weight);
-			}
+					VtxToBoneIdWeightMap[VtxId].SetElement(VertInfluence.BoneIndex, VertInfluence.Weight);
+				}
 			}
 
 			// sort the bones and limit to MaxBonesPerVertex
