Skeletal Mesh Reduction.  Make the lock edges more robust.

diff --git a/Engine/Plugins/Experimental/SkeletalReduction/Source/Private/SkeletalSimplifier.h b/Engine/Plugins/Experimental/SkeletalReduction/Source/Private/SkeletalSimplifier.h
--- a/Engine/Plugins/Experimental/SkeletalReduction/Source/Private/SkeletalSimplifier.h
+++ b/Engine/Plugins/Experimental/SkeletalReduction/Source/Private/SkeletalSimplifier.h
@@ -240,7 +240,7 @@ namespace SkeletalSimplifier
 		int32               CountDegenerates() const; // Included for testing.
 
 
-
+	
 	protected:
 
 		// --- Weights for quadric simplification
@@ -411,10 +411,10 @@ namespace SkeletalSimplifier
 
 		double maxError = 0.0f;
 		float  distError = 0.0f;
-		
+
 		while (CollapseCostHeap.Num() > 0)
 		{
-			
+		
 
 			// get the next vertex to collapse
 			uint32 TopIndex = CollapseCostHeap.Top();
@@ -438,16 +438,15 @@ namespace SkeletalSimplifier
 			// Pointer to the candidate edge (link list) for collapse
 			SimpEdgeType* TopEdgePtr = MeshManager.GetEdgePtr(TopIndex);
 			checkSlow(TopEdgePtr);
-
+			
 
 			// Gather all the edges that are really in this group.
 			EdgePtrArray CoincidentEdges;
 			MeshManager.GetEdgesInGroup(*TopEdgePtr, CoincidentEdges);
 
-			// Will return 'true' if any of the edges are locked.
-			const bool bIsLocked = MeshManager.IsLocked(CoincidentEdges);
-			if (bIsLocked)
+			if (MeshManager.HasLockedVerts(TopEdgePtr) || MeshManager.IsLockedGroup(CoincidentEdges))
 			{
+				// this edge shouldn't be collapsed.
 				continue;
 			}
 
@@ -457,29 +456,52 @@ namespace SkeletalSimplifier
 			MeshManager.GetAdjacentTopology(*TopEdgePtr, DirtyTris, DirtyVerts, DirtyEdges);
 
 			const int32 NumCoincidentEdges = CoincidentEdges.Num();
-#if 1
+
 			// Remove any edges in from this group that happen to be degenerate (no adjacent triangles)
 			// and capture the Idx of the dead edges.  Also removes the edges from the edge-index heap.
 			FSimplifierMeshManager::IdxArray InvalidCostIdxArray;  // Keep track of the Idx to remove from the cost heap
-			MeshManager.RemoveEdgeIfInvalid(CoincidentEdges, InvalidCostIdxArray);
+			const int32 NumRemovedEdges = MeshManager.RemoveEdgeIfInvalid(CoincidentEdges, InvalidCostIdxArray);
+			
+			// if none of the edges in this group were valid, just continue.
+			if (NumCoincidentEdges == NumRemovedEdges)
+			{
+				continue;
+			}
 
+			// The representative edge in our edge group was removed.  Replace it with another from the group that is valid.
+			if (MeshManager.IsRemoved(TopEdgePtr) || MeshManager.IsInvalid(TopEdgePtr))
+			{
+				// try to find a valid edge in the batch.. 
+				for (SimpEdgeType* EPtr : CoincidentEdges)
+				{
+					if (EPtr && !MeshManager.IsRemoved(EPtr))
+					{
+						TopEdgePtr = EPtr;
+						break;
+					}
+				}
+			}
 
-			if (MeshManager.IsInvalid(TopEdgePtr))
+			// continue if no edge actually exists.
+			if (MeshManager.IsRemoved(TopEdgePtr) || MeshManager.IsInvalid(TopEdgePtr))
+			{
 				continue;
-#endif
+			}
 
 			// update Edge->v1 to new locations :  move verts to new verts
 			{
 
 				// Copy the  edge->v1->verts and update the location & attributes
 				// capturing the corresponding SimpVert.
-
+				// The VertexUpdate Array is built by adding 
+				// 1) the two vertices for each edge
+				// 2) Elements from TopEdge->v0 vert group not already added
+				// 3) Elements from TopEdge->v1 vert group not already added
 				EdgeUpdateTupleArray VertexUpdateArray;
-				//ComputeEdgeCollapseVerts(TopEdgePtr, VertexUpdateArray);
 				ComputeEdgeCollapseVertsAndFixBones(TopEdgePtr, VertexUpdateArray); // re-targets using closest bone
-
+			
 				// Compute the distance of the new vertex to each plane of the affected triangles. 
-				
+			
 				if (bCheckDistance)
 				{
 					const FVector NewPos = VertexUpdateArray[0].Get<2>().GetPos();
@@ -507,15 +529,17 @@ namespace SkeletalSimplifier
 					const MeshVertType&  VertAttributes = EdgeUpdate.Get<2>();
 
 					// update the first edge vertex
-					if (SimpVertType* Vert = EdgeUpdate.Get<0>())
+					SimpVertType* VtxPtr = EdgeUpdate.Get<0>();
+					if (VtxPtr != nullptr)
 					{
-						MeshManager.UpdateVertexAttributes(*Vert, VertAttributes);
+						MeshManager.UpdateVertexAttributes(*VtxPtr, VertAttributes);
 					}
 
 					// update the second edge vertex
-					if (SimpVertType* Vert = EdgeUpdate.Get<1>())
+					VtxPtr = EdgeUpdate.Get<1>();
+					if (VtxPtr != nullptr)
 					{
-						MeshManager.UpdateVertexAttributes(*Vert, VertAttributes);
+						MeshManager.UpdateVertexAttributes(*VtxPtr, VertAttributes);
 					}
 
 				}
@@ -523,20 +547,30 @@ namespace SkeletalSimplifier
 
 			// collapse all edges by moving edge->v0 to edge->v1
 			{
+				
+				// All positions and attributes should be fixed now, but we need to update the
+				// mesh connectivity. 
 				for (int i = 0; i < NumCoincidentEdges; ++i)
 				{
 					SimpEdgeType* EdgePtr = CoincidentEdges[i];
 
-					bool bSkip = !EdgePtr || EdgePtr->TestFlags(SIMP_REMOVED);
+					bool bSkip = !EdgePtr || MeshManager.IsRemoved(EdgePtr); 
 
 					if (bSkip) continue;
 
 					// Collapse the edge, delete triangles that become degenerate
 					// and update any edges that used to include v0 to include v1 now.
 					// V1 acquires any lock from v0
-					MeshManager.CollapseEdge(EdgePtr, InvalidCostIdxArray);
+					bool bCollapsed = MeshManager.CollapseEdge(EdgePtr, InvalidCostIdxArray);
 
-					// NB: added 6/29/18. 
+					// this edge has been collapsed.  Remove it from the CoincidentEdges array
+					if (bCollapsed)
+					{
+						CoincidentEdges[i] = NULL;
+					}
+
+					// NB: two edges in this group may have shared a single vertex. 
+					// so the above collapse could have also collapsed other edges in the group.
 					MeshManager.RemoveEdgeIfInvalid(CoincidentEdges, InvalidCostIdxArray);
 
 				}
@@ -545,14 +579,15 @@ namespace SkeletalSimplifier
 
 				{
 					// I'm not totally okay with this.. this adds all the v0 verts to the v1 group,
-					// but doesn't change there Position()
-					MeshManager.MergeGroups(*TopEdgePtr->v0, *TopEdgePtr->v1);
+					// resulting in v1, v1', v1''.. v0, v0', ..
+					// but doesn't change their Position()
+					MeshManager.MergeGroups(TopEdgePtr->v0, TopEdgePtr->v1);  
+
+					// if any of the verts in the group is locked, make them all locked
+					MeshManager.PropagateFlag<SIMP_LOCKED>(*TopEdgePtr->v1);
 
 					// prune any SIMP_REMOVED verts from the v1 link list.
 					MeshManager.PruneVerts<SIMP_REMOVED>(*TopEdgePtr->v1);
-
-					// spread locked flag to all members of the vert group
-					MeshManager.PropagateFlag<SIMP_LOCKED>(*TopEdgePtr->v1);
 				}
 			}
 
@@ -567,10 +602,10 @@ namespace SkeletalSimplifier
 			// If a dirty tri has zero area, remove it and tag as SIMP_REMOVED
 			MeshManager.RemoveIfDegenerate(DirtyTris);
 
-
+			
 			// Tag verts that aren't part of tris ad SIMP_REMOVED
 			// and remove them from the mesh vert groups.
-			MeshManager.RemoveIfDegenerate(DirtyVerts);
+			MeshManager.RemoveIfDegenerate(DirtyVerts); 
 
 			// Remove edges with out verts and update
 			// the cost & heap for all the remaining dirty edges.
diff --git a/Engine/Plugins/Experimental/SkeletalReduction/Source/Private/SkeletalSimplifierMeshManager.cpp b/Engine/Plugins/Experimental/SkeletalReduction/Source/Private/SkeletalSimplifierMeshManager.cpp
--- a/Engine/Plugins/Experimental/SkeletalReduction/Source/Private/SkeletalSimplifierMeshManager.cpp
+++ b/Engine/Plugins/Experimental/SkeletalReduction/Source/Private/SkeletalSimplifierMeshManager.cpp
@@ -319,8 +319,7 @@ void SkeletalSimplifier::FSimplifierMeshManager::GetCoincidentVertGroups(VertPtr
 			while (tmp->next != Vert)
 			{
 				tmp = tmp->next;
-				checkSlow(!tmp->TestFlags(SIMP_REMOVED));
-				if (tmp > maxVert)
+				if (tmp > maxVert && !tmp->TestFlags(SIMP_REMOVED)) 
 				{
 					maxVert = tmp;
 				}
@@ -422,8 +421,7 @@ void SkeletalSimplifier::FSimplifierMeshManager::RebuildEdgeLinkLists(EdgePtrArr
 
 
 void SkeletalSimplifier::FSimplifierMeshManager::FlagBoundary(const ESimpElementFlags Flag)
-{
-	check(Flag == ESimpElementFlags::SIMP_LOCKED);
+{ 
 
 	TArray< SimpVertType*, TInlineAllocator<64> > adjVerts;
 
@@ -434,8 +432,33 @@ void SkeletalSimplifier::FSimplifierMeshManager::FlagBoundary(const ESimpElement
 		checkSlow(v0 != NULL);
 		check(v0->adjTris.Num() > 0);
 
+		if (v0->TestFlags(Flag))
+		{
+			// we must have visited this vert already in a vert group
+			continue;
+		}
+
+
+		//Find all the verts that are adjacent to any vert in this group.
 		adjVerts.Reset();
-		v0->FindAdjacentVertsGroup(adjVerts);
+		//the scope below replaces  v0->FindAdjacentVertsGroup(adjVerts);
+		{
+			SimpVertType* v = v0;
+			do {
+				for (TriIterator triIter = v->adjTris.Begin(); triIter != v->adjTris.End(); ++triIter)
+				{
+					for (int j = 0; j < 3; j++)
+					{
+						SimpVertType* TriVert = (*triIter)->verts[j];
+						if (TriVert != v)
+						{
+							adjVerts.AddUnique(TriVert);
+						}
+					}
+				}
+				v = v->next;
+			} while (v != v0);
+		}
 
 		for (SimpVertType* v1 : adjVerts)
 		{
@@ -459,6 +482,7 @@ void SkeletalSimplifier::FSimplifierMeshManager::FlagBoundary(const ESimpElement
 					vert = vert->next;
 				} while (vert != v0);
 
+				// reset v0-group flag.
 				v0->DisableAdjTriFlagsGroup(SIMP_MARK1);
 
 				if (faceCount == 1)
@@ -652,26 +676,48 @@ int32 SkeletalSimplifier::FSimplifierMeshManager::RemoveEdgeIfInvalid(EdgePtrArr
 	const uint32 NumCandidateEdges = CandidateEdges.Num();
 	for (uint32 i = 0; i < NumCandidateEdges; ++i)
 	{
+	
 		SimpEdgeType* EdgePtr = CandidateEdges[i];
-		
-		// DJH - added 6/29/18
+
+		// Edge has already been removed..
 		if (!EdgePtr) continue;
+
 		
-		if (IsInvalid(EdgePtr))
+		// Verify the edge has an adjacent face.
+		auto HasAdjacentFace = [](SimpEdgeType* e)->bool
 		{
-			const uint32 Idx = RemoveEdge(*EdgePtr);
+			// Verify that this is truly an edge of a triangle
+
+			e->v0->EnableAdjVertFlags(SIMP_MARK1);
+			e->v1->DisableAdjVertFlags(SIMP_MARK1);
+
+			if (e->v0->TestFlags(SIMP_MARK1))
+			{
+				// Invalid edge results from collapsing a bridge tri
+				// There are no actual triangles connecting these verts
+				e->v0->DisableAdjVertFlags(SIMP_MARK1);
+
+				return false;
+			}
+
+			return true;
+		};
+
+		// remove invalid faces from the hash, edge group and flag.
+		if ( IsInvalid(EdgePtr) || !HasAdjacentFace(EdgePtr) ) // one of the verts touches zero triangles
+		{
+			// unlinks the edge from the edge group, remove from edge hash, add removed flag.
+
+			const uint32 Idx = RemoveEdge(*EdgePtr); 
 
 			// Record the index of the edge we remove.
+
 			if (Idx < INVALID_EDGE_ID)
 			{
-				RemovedEdgeIdxArray.AddUnique(Idx); // djh changed from Add()
+				RemovedEdgeIdxArray.AddUnique(Idx); 
 			}
 			CandidateEdges[i] = NULL;
 		}
-		else
-		{
-			checkSlow(!EdgePtr->TestFlags(SIMP_REMOVED));
-		}
 	}
 
 	return RemovedEdgeIdxArray.Num();
@@ -824,12 +870,12 @@ int32 SkeletalSimplifier::FSimplifierMeshManager::RemoveIfDegenerate(VertPtrArra
 	{
 		if (VertPtr->TestFlags(SIMP_REMOVED))
 			continue;
-
+		
 		if (VertPtr->adjTris.Num() == 0)
 		{
 			NumRemoved++;
 			VertPtr->EnableFlags(SIMP_REMOVED);
-
+		
 			// ungroup
 			VertPtr->prev->next = VertPtr->next;
 			VertPtr->next->prev = VertPtr->prev;
@@ -884,7 +930,7 @@ int32 SkeletalSimplifier::FSimplifierMeshManager::RemoveIfDegenerate(EdgePtrArra
 
 		if (edge->TestFlags(SIMP_REMOVED))
 			continue;
-
+	
 		if (edge->v0 == edge->v1)
 		{
 			edge->EnableFlags(SIMP_REMOVED); // djh 8/3/18.  not sure why this happens
@@ -892,19 +938,17 @@ int32 SkeletalSimplifier::FSimplifierMeshManager::RemoveIfDegenerate(EdgePtrArra
 			uint32 Idx = RemoveEdge(*edge);
 			if (Idx < INVALID_EDGE_ID)
 			{
-				RemovedEdgeIdxArray.Add(Idx);
+				RemovedEdgeIdxArray.AddUnique(Idx);
 			}
-
 		}
-
-		if (edge->v0->TestFlags(SIMP_REMOVED) ||
-			edge->v1->TestFlags(SIMP_REMOVED))
+		else if (edge->v0->TestFlags(SIMP_REMOVED) ||
+			     edge->v1->TestFlags(SIMP_REMOVED))
 		{
-			//RemoveEdgeFromLocationMapAndHeap(edge);
+			
 			uint32 Idx = RemoveEdge(*edge);
 			if (Idx < INVALID_EDGE_ID)
 			{
-				RemovedEdgeIdxArray.Add(Idx);
+				RemovedEdgeIdxArray.AddUnique(Idx);
 			}
 		}
 	}
@@ -914,7 +958,7 @@ int32 SkeletalSimplifier::FSimplifierMeshManager::RemoveIfDegenerate(EdgePtrArra
 
 
 
-void SkeletalSimplifier::FSimplifierMeshManager::CollapseEdge(SimpEdgeType * EdgePtr, IdxArray& RemovedEdgeIdxArray)
+bool SkeletalSimplifier::FSimplifierMeshManager::CollapseEdge(SimpEdgeType * EdgePtr, IdxArray& RemovedEdgeIdxArray)
 {
 	SimpVertType* v0 = EdgePtr->v0;
 	SimpVertType* v1 = EdgePtr->v1;
@@ -942,7 +986,16 @@ void SkeletalSimplifier::FSimplifierMeshManager::CollapseEdge(SimpEdgeType * Edg
 		// Invalid edge results from collapsing a bridge tri
 		// There are no actual triangles connecting these verts
 		v0->DisableAdjVertFlags(SIMP_MARK1);
-		return;
+
+		EdgePtr->EnableFlags(SIMP_REMOVED);
+		const uint32 Idx = RemoveEdge(*EdgePtr);
+		if (Idx < INVALID_EDGE_ID)
+		{
+			RemovedEdgeIdxArray.AddUnique(Idx);
+		}
+
+		// return false because the was no real edge to collapse
+		return false;
 	}
 
 	// update edges from v0 to v1
@@ -952,6 +1005,8 @@ void SkeletalSimplifier::FSimplifierMeshManager::CollapseEdge(SimpEdgeType * Edg
 	if (v0->TestFlags(SIMP_LOCKED))
 		v1->EnableFlags(SIMP_LOCKED);
 
+	// this version of the vertex will be removed after the collapse
+	v0->DisableFlags(SIMP_LOCKED); // we already shared the locked state with the remaining vertex
 
 	// Update 'other'-u edges to 'other'-v edges ( where other != v) 
 
@@ -1074,6 +1129,8 @@ void SkeletalSimplifier::FSimplifierMeshManager::CollapseEdge(SimpEdgeType * Edg
 	// record the reduced number of verts
 
 	ReducedNumVerts--;
+
+	return true;
 }
 
 void SkeletalSimplifier::FSimplifierMeshManager::OutputMesh(MeshVertType* verts, uint32* indexes, TArray<int32>* LockedVerts)
diff --git a/Engine/Plugins/Experimental/SkeletalReduction/Source/Private/SkeletalSimplifierMeshManager.h b/Engine/Plugins/Experimental/SkeletalReduction/Source/Private/SkeletalSimplifierMeshManager.h
--- a/Engine/Plugins/Experimental/SkeletalReduction/Source/Private/SkeletalSimplifierMeshManager.h
+++ b/Engine/Plugins/Experimental/SkeletalReduction/Source/Private/SkeletalSimplifierMeshManager.h
@@ -125,7 +125,12 @@ namespace SkeletalSimplifier
 		// Return true if either vertex has no associated faces.
 		bool IsInvalid(const SimpEdgeType* EdgePtr) const
 		{
-			return (EdgePtr->v0->adjTris.Num() == 0 || EdgePtr->v1->adjTris.Num() == 0);
+			 return (EdgePtr->v0->adjTris.Num() == 0 || EdgePtr->v1->adjTris.Num() == 0);
+		}
+
+		bool IsRemoved(const SimpEdgeType* EdgePtr) const
+		{
+			return EdgePtr->TestFlags(SIMP_REMOVED);
 		}
 
 		SimpEdgeType* GetEdgePtr(const uint32 Idx)
@@ -176,10 +181,9 @@ namespace SkeletalSimplifier
 		// Merge two link-lits of verts into a single group.
 		// NB: this does not verify that the locations of the verts are all the same.
 
-		void MergeGroups(SimpVertType& MemberOfGroupA, SimpVertType& MemberOfGroupB)
+		void MergeGroups(SimpVertType*  A, SimpVertType*  B)
 		{
-			SimpVertType* A = &MemberOfGroupA;
-			SimpVertType* B = &MemberOfGroupB;
+			// This inserts the B-loop between A and A.next.
 
 			// combine v0 and v1 groups
 			A->next->prev = B->prev;
@@ -187,6 +191,22 @@ namespace SkeletalSimplifier
 
 			A->next = B;
 			B->prev = A;
+			
+			/**
+			// This inserts the B-loop between A.prev and A.
+			// the last B in the B-loop
+			SimpVertType* Blast = B->prev;
+			// the last A in the A-Loop
+			SimpVertType* Alast = A->prev;
+
+			// Add A-loop after the last B
+			Blast->next = A;
+			A->prev = Blast;
+
+			// make the last 
+			B->prev = Alast;
+			Alast->next = B;
+			*/
 		}
 
 		// If a member of the group has the given flag, propagate that flag
@@ -224,21 +244,22 @@ namespace SkeletalSimplifier
 		void GetVertsInGroup(const SimpVertType& seedVert, VertPtrArray& InOutVertGroup) const
 		{
 			SimpVertType* VertPtr = const_cast<SimpVertType*>(&seedVert);
+			SimpVertType* v0 = VertPtr;
 			do {
-				InOutVertGroup.Add(VertPtr);
-				VertPtr = VertPtr->next;
-			} while (VertPtr != &seedVert);
+				InOutVertGroup.Add(v0);
+				v0 = v0->next;
+			} while (v0 != VertPtr);
 		}
 
-		// Remove any verts tagged as with FlagValue (e.g. SIMP_REMOVED)
+		// Remove any verts tagged as with FlagValue (e.g. SIMP_REMOVED)  NB: this may unlink the seed vert!
 		template < ESimpElementFlags FlagValue >
-		void PruneVerts(const SimpVertType& seedVert)
+		void PruneVerts(SimpVertType& seedVert)
 		{
 			// Get all the verts that are link-listed together
 			VertPtrArray VertsInVertGroup;
 			GetVertsInGroup(seedVert, VertsInVertGroup);
 
-			// Un link any verts that are maked as SIMP_REMOVED
+			// Unlink any verts that are marked as SIMP_REMOVED
 			for (int32 i = 0, Imax = VertsInVertGroup.Num(); i < Imax; ++i)
 			{
 				SimpVertType* v = VertsInVertGroup[i];
@@ -249,6 +270,9 @@ namespace SkeletalSimplifier
 					v->next->prev = v->prev;
 					v->next = v;
 					v->prev = v;
+
+					// cleanup.  Insure that pruned verts aren't marked as locekd.
+					v->DisableFlags(SIMP_LOCKED);
 				}
 			}
 
@@ -318,7 +342,8 @@ namespace SkeletalSimplifier
 		// to a single edge
 		// e.g. collapse edge 0-1 in the triangle 0-1,  1-2, 2-0
 		//      will result in the single edge 1-2
-		void CollapseEdge(SimpEdgeType* EdgePtr, IdxArray& RemovedEdgeIdxArray);
+		// @returns true if the edge collapsed, false if the edge wasn't a true edge.
+		bool CollapseEdge(SimpEdgeType* EdgePtr, IdxArray& RemovedEdgeIdxArray);
 
 		// Mark a tri as removed, and remove it from vertex adj lists.
 		uint32 RemoveTri(SimpTriType& Tri)
@@ -357,7 +382,7 @@ namespace SkeletalSimplifier
 
 
 		// @return true if any of the edges in this edge group is locked.
-		bool IsLocked(const EdgePtrArray& EdgeGroup) const
+		bool IsLockedGroup(const EdgePtrArray& EdgeGroup) const
 		{
 			bool locked = false;
 			int32 NumEdgesInGroup = EdgeGroup.Num();
@@ -374,6 +399,13 @@ namespace SkeletalSimplifier
 			return locked;
 		}
 
+		//@return true if the either of the verts are locked
+		bool HasLockedVerts(const SimpEdgeType* Edge) const
+		{
+			return (Edge->v0->TestFlags(SIMP_LOCKED) || Edge->v1->TestFlags(SIMP_LOCKED));
+		}
+
+
 		// Find the edge associated with these verts.  Will return NULL
 		// if no such edge exists.
 		SimpEdgeType* FindEdge(const SimpVertType* u, const SimpVertType* v)
