Copying //UE4/Dev-Editor to Dev-Main (//UE4/Dev-Main) for 4.24

diff --git a/Engine/Plugins/Experimental/ModelingToolsEditorMode/Source/ModelingToolsEditorMode/Private/ModelingToolsEditorModeStyle.cpp b/Engine/Plugins/Experimental/ModelingToolsEditorMode/Source/ModelingToolsEditorMode/Private/ModelingToolsEditorModeStyle.cpp
new file mode 100644
--- /dev/null
+++ b/Engine/Plugins/Experimental/ModelingToolsEditorMode/Source/ModelingToolsEditorMode/Private/ModelingToolsEditorModeStyle.cpp
@@ -0,0 +1,215 @@
+// Copyright 1998-2019 Epic Games, Inc. All Rights Reserved.
+
+#include "ModelingToolsEditorModeStyle.h"
+#include "Styling/SlateStyleRegistry.h"
+#include "Styling/SlateTypes.h"
+#include "Styling/CoreStyle.h"
+#include "EditorStyleSet.h"
+#include "Interfaces/IPluginManager.h"
+#include "SlateOptMacros.h"
+
+
+#define IMAGE_PLUGIN_BRUSH( RelativePath, ... ) FSlateImageBrush( FModelingToolsEditorModeStyle::InContent( RelativePath, ".png" ), __VA_ARGS__ )
+#define IMAGE_BRUSH(RelativePath, ...) FSlateImageBrush(StyleSet->RootToContentDir(RelativePath, TEXT(".png")), __VA_ARGS__)
+#define BOX_BRUSH(RelativePath, ...) FSlateBoxBrush(StyleSet->RootToContentDir(RelativePath, TEXT(".png")), __VA_ARGS__)
+#define DEFAULT_FONT(...) FCoreStyle::GetDefaultFontStyle(__VA_ARGS__)
+
+FString FModelingToolsEditorModeStyle::InContent(const FString& RelativePath, const ANSICHAR* Extension)
+{
+	static FString ContentDir = IPluginManager::Get().FindPlugin(TEXT("ModelingToolsEditorMode"))->GetContentDir();
+	return (ContentDir / RelativePath) + Extension;
+}
+
+TSharedPtr< FSlateStyleSet > FModelingToolsEditorModeStyle::StyleSet = nullptr;
+TSharedPtr< class ISlateStyle > FModelingToolsEditorModeStyle::Get() { return StyleSet; }
+
+FName FModelingToolsEditorModeStyle::GetStyleSetName()
+{
+	static FName ModelingToolsStyleName(TEXT("ModelingToolsStyle"));
+	return ModelingToolsStyleName;
+}
+
+BEGIN_SLATE_FUNCTION_BUILD_OPTIMIZATION
+
+void FModelingToolsEditorModeStyle::Initialize()
+{
+	// Const icon sizes
+	const FVector2D Icon8x8(8.0f, 8.0f);
+	const FVector2D Icon16x16(16.0f, 16.0f);
+	const FVector2D Icon20x20(20.0f, 20.0f);
+	const FVector2D Icon28x28(28.0f, 28.0f);
+	const FVector2D Icon40x40(40.0f, 40.0f);
+
+	// Only register once
+	if (StyleSet.IsValid())
+	{
+		return;
+	}
+
+	StyleSet = MakeShareable(new FSlateStyleSet(GetStyleSetName()));
+	StyleSet->SetContentRoot(FPaths::EngineContentDir() / TEXT("Editor/Slate"));
+	StyleSet->SetCoreContentRoot(FPaths::EngineContentDir() / TEXT("Slate"));
+
+	const FTextBlockStyle& NormalText = FEditorStyle::Get().GetWidgetStyle<FTextBlockStyle>("NormalText");
+
+	// Shared editors
+	//{
+	//	StyleSet->Set("Paper2D.Common.ViewportZoomTextStyle", FTextBlockStyle(NormalText)
+	//		.SetFont(DEFAULT_FONT("BoldCondensed", 16))
+	//	);
+
+	//	StyleSet->Set("Paper2D.Common.ViewportTitleTextStyle", FTextBlockStyle(NormalText)
+	//		.SetFont(DEFAULT_FONT("Regular", 18))
+	//		.SetColorAndOpacity(FLinearColor(1.0, 1.0f, 1.0f, 0.5f))
+	//	);
+
+	//	StyleSet->Set("Paper2D.Common.ViewportTitleBackground", new BOX_BRUSH("Old/Graph/GraphTitleBackground", FMargin(0)));
+	//}
+
+	// Tool Manager icons
+	{
+		// Accept/Cancel/Complete active tool
+		StyleSet->Set("ModelingToolsManagerCommands.CancelActiveTool", new IMAGE_PLUGIN_BRUSH("Icons/icon_ActiveTool_Cancel_40x", Icon28x28));
+		StyleSet->Set("ModelingToolsManagerCommands.CancelActiveTool.Small", new IMAGE_PLUGIN_BRUSH("Icons/icon_ActiveTool_Cancel_40x", Icon28x28));
+		StyleSet->Set("ModelingToolsManagerCommands.AcceptActiveTool", new IMAGE_PLUGIN_BRUSH("Icons/icon_ActiveTool_Accept_40x", Icon28x28));
+		StyleSet->Set("ModelingToolsManagerCommands.AcceptActiveTool.Small", new IMAGE_PLUGIN_BRUSH("Icons/icon_ActiveTool_Accept_40x", Icon28x28));
+		StyleSet->Set("ModelingToolsManagerCommands.CompleteActiveTool", new IMAGE_PLUGIN_BRUSH("Icons/icon_ActiveTool_Accept_40x", Icon28x28));
+		StyleSet->Set("ModelingToolsManagerCommands.CompleteActiveTool.Small", new IMAGE_PLUGIN_BRUSH("Icons/icon_ActiveTool_Accept_40x", Icon28x28));
+
+
+		StyleSet->Set("ModelingToolsManagerCommands.BeginAddPrimitiveTool", new IMAGE_PLUGIN_BRUSH("Icons/icon_Tool_MakePrimitive_40x", Icon28x28));
+		StyleSet->Set("ModelingToolsManagerCommands.BeginAddPrimitiveTool.Small", new IMAGE_PLUGIN_BRUSH("Icons/icon_Tool_MakePrimitive_40x", Icon28x28));
+		StyleSet->Set("ModelingToolsManagerCommands.BeginDrawPolygonTool", new IMAGE_PLUGIN_BRUSH("Icons/icon_Tool_DrawPolygon_40x", Icon28x28));
+		StyleSet->Set("ModelingToolsManagerCommands.BeginDrawPolygonTool.Small", new IMAGE_PLUGIN_BRUSH("Icons/icon_Tool_DrawPolygon_40x", Icon28x28));
+		StyleSet->Set("ModelingToolsManagerCommands.BeginShapeSprayTool", new IMAGE_PLUGIN_BRUSH("Icons/icon_Tool_ShapeSpray_40x", Icon28x28));
+		StyleSet->Set("ModelingToolsManagerCommands.BeginShapeSprayTool.Small", new IMAGE_PLUGIN_BRUSH("Icons/icon_Tool_ShapeSpray_40x", Icon28x28));
+
+		StyleSet->Set("ModelingToolsManagerCommands.BeginSmoothMeshTool", new IMAGE_PLUGIN_BRUSH("Icons/icon_Tool_Smooth_40x", Icon28x28));
+		StyleSet->Set("ModelingToolsManagerCommands.BeginSmoothMeshTool.Small", new IMAGE_PLUGIN_BRUSH("Icons/icon_Tool_Smooth_40x", Icon28x28));
+		StyleSet->Set("ModelingToolsManagerCommands.BeginSculptMeshTool", new IMAGE_PLUGIN_BRUSH("Icons/icon_Tool_Sculpt_40x", Icon28x28));
+		StyleSet->Set("ModelingToolsManagerCommands.BeginSculptMeshTool.Small", new IMAGE_PLUGIN_BRUSH("Icons/icon_Tool_Sculpt_40x", Icon28x28));
+		StyleSet->Set("ModelingToolsManagerCommands.BeginPolyEditTool", new IMAGE_PLUGIN_BRUSH("Icons/icon_Tool_PolyEdit_40x", Icon28x28));
+		StyleSet->Set("ModelingToolsManagerCommands.BeginPolyEditTool.Small", new IMAGE_PLUGIN_BRUSH("Icons/icon_Tool_PolyEdit_40x", Icon28x28));
+		StyleSet->Set("ModelingToolsManagerCommands.BeginDisplaceMeshTool", new IMAGE_PLUGIN_BRUSH("Icons/icon_Tool_Displace_40x", Icon28x28));
+		StyleSet->Set("ModelingToolsManagerCommands.BeginDisplaceMeshTool.Small", new IMAGE_PLUGIN_BRUSH("Icons/icon_Tool_Displace_40x", Icon28x28));
+		StyleSet->Set("ModelingToolsManagerCommands.BeginMeshSpaceDeformerTool", new IMAGE_PLUGIN_BRUSH("Icons/icon_Tool_Displace_40x", Icon28x28));
+		StyleSet->Set("ModelingToolsManagerCommands.BeginMeshSpaceDeformerTool.Small", new IMAGE_PLUGIN_BRUSH("Icons/icon_Tool_Displace_40x", Icon28x28));
+		StyleSet->Set("ModelingToolsManagerCommands.BeginTransformMeshesTool", new IMAGE_PLUGIN_BRUSH("Icons/icon_Tool_Transform_40x", Icon28x28));
+		StyleSet->Set("ModelingToolsManagerCommands.BeginTransformMeshesTool.Small", new IMAGE_PLUGIN_BRUSH("Icons/icon_Tool_Transform_40x", Icon28x28));
+
+		StyleSet->Set("ModelingToolsManagerCommands.BeginRemeshSculptMeshTool", new IMAGE_PLUGIN_BRUSH("Icons/icon_Tool_RemeshSculpt_40x", Icon28x28));
+		StyleSet->Set("ModelingToolsManagerCommands.BeginRemeshSculptMeshTool.Small", new IMAGE_PLUGIN_BRUSH("Icons/icon_Tool_RemeshSculpt_40x", Icon28x28));
+		StyleSet->Set("ModelingToolsManagerCommands.BeginRemeshMeshTool", new IMAGE_PLUGIN_BRUSH("Icons/icon_Tool_Remesh_40x", Icon28x28));
+		StyleSet->Set("ModelingToolsManagerCommands.BeginRemeshMeshTool.Small", new IMAGE_PLUGIN_BRUSH("Icons/icon_Tool_Remesh_40x", Icon28x28));
+		StyleSet->Set("ModelingToolsManagerCommands.BeginSimplifyMeshTool", new IMAGE_PLUGIN_BRUSH("Icons/icon_Tool_Simplify_40x", Icon28x28));
+		StyleSet->Set("ModelingToolsManagerCommands.BeginSimplifyMeshTool.Small", new IMAGE_PLUGIN_BRUSH("Icons/icon_Tool_Simplify_40x", Icon28x28));
+		StyleSet->Set("ModelingToolsManagerCommands.BeginEditNormalsTool", new IMAGE_PLUGIN_BRUSH("Icons/icon_Tool_EditNormals_40x", Icon28x28));
+		StyleSet->Set("ModelingToolsManagerCommands.BeginEditNormalsTool.Small", new IMAGE_PLUGIN_BRUSH("Icons/icon_Tool_EditNormals_40x", Icon28x28));
+		StyleSet->Set("ModelingToolsManagerCommands.BeginUVProjectionTool", new IMAGE_PLUGIN_BRUSH("Icons/icon_Tool_UVProjection_40x", Icon28x28));
+		StyleSet->Set("ModelingToolsManagerCommands.BeginUVProjectionTool.Small", new IMAGE_PLUGIN_BRUSH("Icons/icon_Tool_UVProjection_40x", Icon28x28));
+		StyleSet->Set("ModelingToolsManagerCommands.BeginVoxelMergeTool", new IMAGE_PLUGIN_BRUSH("Icons/icon_Tool_VoxMerge_40x", Icon28x28));
+		StyleSet->Set("ModelingToolsManagerCommands.BeginVoxelMergeTool.Small", new IMAGE_PLUGIN_BRUSH("Icons/icon_Tool_VoxMerge_40x", Icon28x28));
+		StyleSet->Set("ModelingToolsManagerCommands.BeginVoxelBooleanTool", new IMAGE_PLUGIN_BRUSH("Icons/icon_Tool_VoxBoolean_40x", Icon28x28));
+		StyleSet->Set("ModelingToolsManagerCommands.BeginVoxelBooleanTool.Small", new IMAGE_PLUGIN_BRUSH("Icons/icon_Tool_VoxBoolean_40x", Icon28x28));
+		StyleSet->Set("ModelingToolsManagerCommands.BeginPlaneCutTool", new IMAGE_PLUGIN_BRUSH("Icons/icon_Tool_PlaneCut_40x", Icon28x28));
+		StyleSet->Set("ModelingToolsManagerCommands.BeginPlaneCutTool.Small", new IMAGE_PLUGIN_BRUSH("Icons/icon_Tool_PlaneCut_40x", Icon28x28));
+		StyleSet->Set("ModelingToolsManagerCommands.BeginPolygonOnMeshTool", new IMAGE_PLUGIN_BRUSH("Icons/icon_Tool_PolygonOnMesh_40x", Icon28x28));
+		StyleSet->Set("ModelingToolsManagerCommands.BeginPolygonOnMeshTool.Small", new IMAGE_PLUGIN_BRUSH("Icons/icon_Tool_PolygonOnMesh_40x", Icon28x28));
+		StyleSet->Set("ModelingToolsManagerCommands.BeginMeshSelectionTool", new IMAGE_PLUGIN_BRUSH("Icons/icon_Tool_MeshSelection_40x", Icon28x28));
+		StyleSet->Set("ModelingToolsManagerCommands.BeginMeshSelectionTool.Small", new IMAGE_PLUGIN_BRUSH("Icons/icon_Tool_MeshSelection_40x", Icon28x28));
+
+		StyleSet->Set("ModelingToolsManagerCommands.BeginMeshInspectorTool", new IMAGE_PLUGIN_BRUSH("Icons/icon_Tool_MeshInspector_40x", Icon28x28));
+		StyleSet->Set("ModelingToolsManagerCommands.BeginMeshInspectorTool.Small", new IMAGE_PLUGIN_BRUSH("Icons/icon_Tool_MeshInspector_40x", Icon28x28));
+		StyleSet->Set("ModelingToolsManagerCommands.BeginParameterizeMeshTool", new IMAGE_PLUGIN_BRUSH("Icons/icon_Tool_UVGenerate_40x", Icon28x28));
+		StyleSet->Set("ModelingToolsManagerCommands.BeginParameterizeMeshTool.Small", new IMAGE_PLUGIN_BRUSH("Icons/icon_Tool_UVGenerate_40x", Icon28x28));
+		StyleSet->Set("ModelingToolsManagerCommands.BeginWeldEdgesTool", new IMAGE_PLUGIN_BRUSH("Icons/icon_Tool_WeldEdges_40x", Icon28x28));
+		StyleSet->Set("ModelingToolsManagerCommands.BeginWeldEdgesTool.Small", new IMAGE_PLUGIN_BRUSH("Icons/icon_Tool_WeldEdges_40x", Icon28x28));
+		StyleSet->Set("ModelingToolsManagerCommands.BeginPolyGroupsTool", new IMAGE_PLUGIN_BRUSH("Icons/icon_Tool_PolyGroups_40x", Icon28x28));
+		StyleSet->Set("ModelingToolsManagerCommands.BeginPolyGroupsTool.Small", new IMAGE_PLUGIN_BRUSH("Icons/icon_Tool_PolyGroups_40x", Icon28x28));
+		StyleSet->Set("ModelingToolsManagerCommands.BeginAttributeEditorTool", new IMAGE_PLUGIN_BRUSH("Icons/icon_Tool_AttributeEditor_40x", Icon28x28));
+		StyleSet->Set("ModelingToolsManagerCommands.BeginAttributeEditorTool.Small", new IMAGE_PLUGIN_BRUSH("Icons/icon_Tool_AttributeEditor_40x", Icon28x28));
+
+
+
+		//StyleSet->Set("ModelingToolsManagerCommands.CancelActiveTool", new IMAGE_PLUGIN_BRUSH("Icons/icon_ActiveTool_Cancel_40x", Icon40x40));
+		//StyleSet->Set("ModelingToolsManagerCommands.CancelActiveTool.Small", new IMAGE_PLUGIN_BRUSH("Icons/icon_ActiveTool_Cancel_20x", Icon20x20));
+		//StyleSet->Set("ModelingToolsManagerCommands.AcceptActiveTool", new IMAGE_PLUGIN_BRUSH("Icons/icon_ActiveTool_Accept_40x", Icon40x40));
+		//StyleSet->Set("ModelingToolsManagerCommands.AcceptActiveTool.Small", new IMAGE_PLUGIN_BRUSH("Icons/icon_ActiveTool_Accept_20x", Icon20x20));
+		//StyleSet->Set("ModelingToolsManagerCommands.CompleteActiveTool", new IMAGE_PLUGIN_BRUSH("Icons/icon_ActiveTool_Accept_40x", Icon40x40));
+		//StyleSet->Set("ModelingToolsManagerCommands.CompleteActiveTool.Small", new IMAGE_PLUGIN_BRUSH("Icons/icon_ActiveTool_Accept_20x", Icon20x20));
+
+
+		//StyleSet->Set("ModelingToolsManagerCommands.BeginMakePrimitiveTool", new IMAGE_PLUGIN_BRUSH("Icons/icon_Tool_MakePrimitive_40x", Icon40x40));
+		//StyleSet->Set("ModelingToolsManagerCommands.BeginMakePrimitiveTool.Small", new IMAGE_PLUGIN_BRUSH("Icons/icon_Tool_MakePrimitive_20x", Icon20x20));
+		//StyleSet->Set("ModelingToolsManagerCommands.BeginDrawPolygonTool", new IMAGE_PLUGIN_BRUSH("Icons/icon_Tool_DrawPolygon_40x", Icon40x40));
+		//StyleSet->Set("ModelingToolsManagerCommands.BeginDrawPolygonTool.Small", new IMAGE_PLUGIN_BRUSH("Icons/icon_Tool_DrawPolygon_20x", Icon20x20));
+		//StyleSet->Set("ModelingToolsManagerCommands.BeginShapeSprayTool", new IMAGE_PLUGIN_BRUSH("Icons/icon_Tool_ShapeSpray_40x", Icon40x40));
+		//StyleSet->Set("ModelingToolsManagerCommands.BeginShapeSprayTool.Small", new IMAGE_PLUGIN_BRUSH("Icons/icon_Tool_ShapeSpray_20x", Icon20x20));
+
+		//StyleSet->Set("ModelingToolsManagerCommands.BeginSmoothMeshTool", new IMAGE_PLUGIN_BRUSH("Icons/icon_Tool_Smooth_40x", Icon40x40));
+		//StyleSet->Set("ModelingToolsManagerCommands.BeginSmoothMeshTool.Small", new IMAGE_PLUGIN_BRUSH("Icons/icon_Tool_Smooth_20x", Icon20x20));
+		//StyleSet->Set("ModelingToolsManagerCommands.BeginSculptMeshTool", new IMAGE_PLUGIN_BRUSH("Icons/icon_Tool_Sculpt_40x", Icon40x40));
+		//StyleSet->Set("ModelingToolsManagerCommands.BeginSculptMeshTool.Small", new IMAGE_PLUGIN_BRUSH("Icons/icon_Tool_Sculpt_20x", Icon20x20));
+		//StyleSet->Set("ModelingToolsManagerCommands.BeginPolyEditTool", new IMAGE_PLUGIN_BRUSH("Icons/icon_Tool_PolyEdit_40x", Icon40x40));
+		//StyleSet->Set("ModelingToolsManagerCommands.BeginPolyEditTool.Small", new IMAGE_PLUGIN_BRUSH("Icons/icon_Tool_PolyEdit_20x", Icon20x20));
+
+		//StyleSet->Set("ModelingToolsManagerCommands.BeginRemeshSculptMeshTool", new IMAGE_PLUGIN_BRUSH("Icons/icon_Tool_RemeshSculpt_40x", Icon40x40));
+		//StyleSet->Set("ModelingToolsManagerCommands.BeginRemeshSculptMeshTool.Small", new IMAGE_PLUGIN_BRUSH("Icons/icon_Tool_RemeshSculpt_20x", Icon20x20));
+		//StyleSet->Set("ModelingToolsManagerCommands.BeginRemeshMeshTool", new IMAGE_PLUGIN_BRUSH("Icons/icon_Tool_Remesh_40x", Icon40x40));
+		//StyleSet->Set("ModelingToolsManagerCommands.BeginRemeshMeshTool.Small", new IMAGE_PLUGIN_BRUSH("Icons/icon_Tool_Remesh_20x", Icon20x20));
+		//StyleSet->Set("ModelingToolsManagerCommands.BeginSimplifyMeshTool", new IMAGE_PLUGIN_BRUSH("Icons/icon_Tool_Simplify_40x", Icon40x40));
+		//StyleSet->Set("ModelingToolsManagerCommands.BeginSimplifyMeshTool.Small", new IMAGE_PLUGIN_BRUSH("Icons/icon_Tool_Simplify_20x", Icon20x20));
+
+
+		//StyleSet->Set("ModelingToolsManagerCommands.BeginMeshInspectorTool", new IMAGE_PLUGIN_BRUSH("Icons/icon_Tool_MeshInspector_40x", Icon40x40));
+		//StyleSet->Set("ModelingToolsManagerCommands.BeginMeshInspectorTool.Small", new IMAGE_PLUGIN_BRUSH("Icons/icon_Tool_MeshInspector_20x", Icon20x20));
+		//StyleSet->Set("ModelingToolsManagerCommands.BeginWeldEdgesTool", new IMAGE_PLUGIN_BRUSH("Icons/icon_Tool_WeldEdges_40x", Icon40x40));
+		//StyleSet->Set("ModelingToolsManagerCommands.BeginWeldEdgesTool.Small", new IMAGE_PLUGIN_BRUSH("Icons/icon_Tool_WeldEdges_20x", Icon20x20));
+		//StyleSet->Set("ModelingToolsManagerCommands.BeginPolyGroupsTool", new IMAGE_PLUGIN_BRUSH("Icons/icon_Tool_PolyGroups_40x", Icon40x40));
+		//StyleSet->Set("ModelingToolsManagerCommands.BeginPolyGroupsTool.Small", new IMAGE_PLUGIN_BRUSH("Icons/icon_Tool_PolyGroups_20x", Icon20x20));
+
+
+
+
+		//const FLinearColor LayerSelectionColor = FLinearColor(0.13f, 0.70f, 1.00f);
+
+		//// Selection color for the active row should be ??? to align with the editor viewport.
+		//const FTableRowStyle& NormalTableRowStyle = FEditorStyle::Get().GetWidgetStyle<FTableRowStyle>("TableView.Row");
+
+		//StyleSet->Set("TileMapEditor.LayerBrowser.TableViewRow",
+		//	FTableRowStyle(NormalTableRowStyle)
+		//	.SetActiveBrush(IMAGE_BRUSH("Common/Selection", Icon8x8, LayerSelectionColor))
+		//	.SetActiveHoveredBrush(IMAGE_BRUSH("Common/Selection", Icon8x8, LayerSelectionColor))
+		//	.SetInactiveBrush(IMAGE_BRUSH("Common/Selection", Icon8x8, LayerSelectionColor))
+		//	.SetInactiveHoveredBrush(IMAGE_BRUSH("Common/Selection", Icon8x8, LayerSelectionColor))
+		//);
+
+		//StyleSet->Set("TileMapEditor.LayerBrowser.SelectionColor", LayerSelectionColor);
+
+		//StyleSet->Set("TileMapEditor.TileSetPalette.NothingSelectedText", FTextBlockStyle(NormalText)
+		//	.SetFont(DEFAULT_FONT("BoldCondensed", 18))
+		//	.SetColorAndOpacity(FLinearColor(0.8, 0.8f, 0.0f, 0.8f))
+		//	.SetShadowOffset(FVector2D(1.0f, 1.0f))
+		//	.SetShadowColorAndOpacity(FLinearColor(0.0f, 0.0f, 0.0f, 0.9f))
+		//);
+	}
+
+	FSlateStyleRegistry::RegisterSlateStyle(*StyleSet.Get());
+};
+
+END_SLATE_FUNCTION_BUILD_OPTIMIZATION
+
+#undef IMAGE_PLUGIN_BRUSH
+#undef IMAGE_BRUSH
+#undef BOX_BRUSH
+#undef DEFAULT_FONT
+
+void FModelingToolsEditorModeStyle::Shutdown()
+{
+	if (StyleSet.IsValid())
+	{
+		FSlateStyleRegistry::UnRegisterSlateStyle(*StyleSet.Get());
+		ensure(StyleSet.IsUnique());
+		StyleSet.Reset();
+	}
+}
diff --git a/Engine/Plugins/Experimental/ModelingToolsEditorMode/Source/ModelingToolsEditorMode/Private/ModelingToolsManagerActions.cpp b/Engine/Plugins/Experimental/ModelingToolsEditorMode/Source/ModelingToolsEditorMode/Private/ModelingToolsManagerActions.cpp
new file mode 100644
--- /dev/null
+++ b/Engine/Plugins/Experimental/ModelingToolsEditorMode/Source/ModelingToolsEditorMode/Private/ModelingToolsManagerActions.cpp
@@ -0,0 +1,60 @@
+// Copyright 1998-2019 Epic Games, Inc. All Rights Reserved.
+
+#include "ModelingToolsManagerActions.h"
+#include "EditorStyleSet.h"
+#include "ModelingToolsEditorModeStyle.h"
+
+#define LOCTEXT_NAMESPACE "ModelingToolsManagerCommands"
+
+
+
+FModelingToolsManagerCommands::FModelingToolsManagerCommands() :
+	TCommands<FModelingToolsManagerCommands>(
+		"ModelingToolsManagerCommands", // Context name for fast lookup
+		NSLOCTEXT("Contexts", "ModelingToolsEditMode", "Modeling Tools Mode"), // Localized context name for displaying
+		NAME_None, // Parent
+		FModelingToolsEditorModeStyle::Get()->GetStyleSetName() // Icon Style Set
+		)
+{
+}
+
+
+void FModelingToolsManagerCommands::RegisterCommands()
+{
+	UI_COMMAND(BeginAddPrimitiveTool, "Primitive", "Start the Add Primitive Tool", EUserInterfaceActionType::Button, FInputChord());
+	UI_COMMAND(BeginDrawPolygonTool, "Polygon", "Start the Draw Polygon Tool", EUserInterfaceActionType::Button, FInputChord());
+	UI_COMMAND(BeginShapeSprayTool, "ShapeSpray", "Start the Shape Spray Tool", EUserInterfaceActionType::Button, FInputChord());
+
+	UI_COMMAND(BeginSculptMeshTool, "Sculpt", "Start the Sculpt Mesh Tool", EUserInterfaceActionType::Button, FInputChord());
+	UI_COMMAND(BeginPolyEditTool, "PolyEdit", "Start the PolyEdit Tool", EUserInterfaceActionType::Button, FInputChord());
+	UI_COMMAND(BeginSmoothMeshTool, "Smooth", "Start the Smooth Mesh Tool", EUserInterfaceActionType::Button, FInputChord());
+	UI_COMMAND(BeginDisplaceMeshTool, "Displace", "Start the Displace Mesh Tool", EUserInterfaceActionType::Button, FInputChord());
+	UI_COMMAND(BeginMeshSpaceDeformerTool, "SpaceDeform", "Start the Mesh Space Deformer Tool", EUserInterfaceActionType::Button, FInputChord());
+	UI_COMMAND(BeginTransformMeshesTool, "Transform", "Start the Transform Tool", EUserInterfaceActionType::Button, FInputChord());
+
+	UI_COMMAND(BeginRemeshSculptMeshTool, "DynaSculpt", "Start the Liquid Sculpt Mesh Tool", EUserInterfaceActionType::Button, FInputChord());
+	UI_COMMAND(BeginRemeshMeshTool, "Remesh", "Start the Remesh Tool", EUserInterfaceActionType::Button, FInputChord());
+	UI_COMMAND(BeginSimplifyMeshTool, "Simplify", "Start the Simplify Mesh Tool", EUserInterfaceActionType::Button, FInputChord());
+	UI_COMMAND(BeginEditNormalsTool, "Normals", "Start the Edit Normals Tool", EUserInterfaceActionType::Button, FInputChord());
+	UI_COMMAND(BeginUVProjectionTool, "UVProjection", "Start the UV Projection Tool", EUserInterfaceActionType::Button, FInputChord());
+	UI_COMMAND(BeginVoxelMergeTool, "VoxMerge", "Start the Voxel Merge Tool", EUserInterfaceActionType::Button, FInputChord());
+	UI_COMMAND(BeginVoxelBooleanTool, "VoxBoolean", "Start the Voxel Boolean Tool", EUserInterfaceActionType::Button, FInputChord());
+	UI_COMMAND(BeginPlaneCutTool, "PlaneCut", "Start the Plane Cut Tool", EUserInterfaceActionType::Button, FInputChord());
+	UI_COMMAND(BeginPolygonOnMeshTool, "HoleCut", "Start the Hole Cut Tool", EUserInterfaceActionType::Button, FInputChord());
+	UI_COMMAND(BeginParameterizeMeshTool, "UVGenerate", "Start the Mesh Parameterization Tool", EUserInterfaceActionType::Button, FInputChord());
+	UI_COMMAND(BeginMeshSelectionTool, "Select", "Start the Mesh Selection Tool", EUserInterfaceActionType::Button, FInputChord());
+
+	UI_COMMAND(BeginMeshInspectorTool, "Inspector", "Start the Mesh Inspector Tool", EUserInterfaceActionType::Button, FInputChord());
+	UI_COMMAND(BeginWeldEdgesTool, "Weld Edges", "Start the Weld Edges Tool", EUserInterfaceActionType::Button, FInputChord());
+	UI_COMMAND(BeginPolyGroupsTool, "PolyGroups", "Start the Poly Groups Tool", EUserInterfaceActionType::Button, FInputChord());
+	UI_COMMAND(BeginAttributeEditorTool, "AttributeEditor", "Start the Attribute Editor Tool", EUserInterfaceActionType::Button, FInputChord());
+
+	UI_COMMAND(AcceptActiveTool, "Accept", "Accept the active tool", EUserInterfaceActionType::Button, FInputChord());
+	UI_COMMAND(CancelActiveTool, "Cancel", "Cancel the active tool", EUserInterfaceActionType::Button, FInputChord());
+	UI_COMMAND(CompleteActiveTool, "Complete", "Complete the active tool", EUserInterfaceActionType::Button, FInputChord());
+}
+
+
+
+
+#undef LOCTEXT_NAMESPACE
diff --git a/Engine/Plugins/Experimental/ModelingToolsEditorMode/Source/ModelingToolsEditorMode/Public/ModelingToolsEditorMode.h b/Engine/Plugins/Experimental/ModelingToolsEditorMode/Source/ModelingToolsEditorMode/Public/ModelingToolsEditorMode.h
new file mode 100644
--- /dev/null
+++ b/Engine/Plugins/Experimental/ModelingToolsEditorMode/Source/ModelingToolsEditorMode/Public/ModelingToolsEditorMode.h
@@ -0,0 +1,164 @@
+// Copyright 1998-2019 Epic Games, Inc. All Rights Reserved.
+
+#pragma once
+
+#include "CoreMinimal.h"
+#include "EdMode.h"
+
+#include "InputState.h"
+#include "InteractiveToolManager.h"
+#include "EdModeInteractiveToolsContext.h"
+
+
+class FEditorComponentSourceFactory;
+class FEditorToolAssetAPI;
+class FUICommandList;
+
+class FModelingToolsEditorMode : public FEdMode
+{
+public:
+	const static FEditorModeID EM_ModelingToolsEditorModeId;
+public:
+	FModelingToolsEditorMode();
+	virtual ~FModelingToolsEditorMode();
+
+	////////////////
+	// FEdMode interface
+	////////////////
+
+	virtual void Tick(FEditorViewportClient* ViewportClient, float DeltaTime) override;
+	virtual void Render(const FSceneView* View, FViewport* Viewport, FPrimitiveDrawInterface* PDI) override;
+	virtual bool UsesToolkits() const override;
+
+	virtual void AddReferencedObjects(FReferenceCollector& Collector) override;
+
+	// these disable the standard gizmo, which is probably want we want in
+	// these tools as we can't hit-test the standard gizmo...
+	virtual bool AllowWidgetMove() override;
+	virtual bool ShouldDrawWidget() const override;
+	virtual bool UsesTransformWidget() const override;
+
+	virtual void ActorSelectionChangeNotify() override;
+
+	virtual bool ProcessEditDelete();
+
+
+	/*
+	 * focus events
+	 */
+
+	// called when we "start" this editor mode (ie switch to this tab)
+	virtual void Enter() override;
+
+	// called when we "end" this editor mode (ie switch to another tab)
+	virtual void Exit() override;
+
+	// called when viewport window is focused
+	virtual bool ReceivedFocus(FEditorViewportClient* ViewportClient, FViewport* Viewport) override;
+
+	// called when viewport window loses focus (ie when some other window is focused)
+	// *not* called when editor is backgrounded, but is called when editor is minimized
+	virtual bool LostFocus(FEditorViewportClient* ViewportClient, FViewport* Viewport) override;
+
+
+	/*
+	 * Mouse position events - These are called when no mouse button is down
+	 */
+
+	// called when mouse moves over viewport window
+	virtual bool MouseEnter(FEditorViewportClient* ViewportClient, FViewport* Viewport, int32 x, int32 y) override;
+
+	// called when mouse leaves viewport window
+	virtual bool MouseLeave(FEditorViewportClient* ViewportClient, FViewport* Viewport) override;
+
+	// called on any mouse-move event. *not* called during tracking/capturing, eg if any button is down
+	virtual bool MouseMove(FEditorViewportClient* ViewportClient, FViewport* Viewport, int32 x, int32 y) override;
+
+
+	/*
+	 * Input Button/Axis Events & Mouse Capture
+	 *
+	 *   event sequence for left mouse down/capture/up is:
+	 *     - InputKey( EKeys::LeftMouseButton, IE_PRESSED )
+	 *     - StartTracking()
+	 *     - CapturedMouseMove() and InputAxis() (repeated for each mouse-move)
+	 *     - Inputkey( EKeys::LeftMouseButton, IE_RELEASED )
+	 *     - EndTracking()
+	 *
+	 *   for doubleclick, we get one of the above sequence, then a second
+	 *   where instead of the first IE_PRESSED we get a IE_DoubleClick
+	 *
+	 *   for mouse wheel, we get following sequence
+	 *     - InputKey ( EKeys::MouseScrollUp / Down, IE_PRESSED )
+	 *     - InputAxis() for wheel move
+	 *     - InputKey ( EKeys::MouseScrollUp / Down, IE_RELEASED )
+	 *   it appears that there will only ever be *one* InputAxis() between the pressed/released sequenece
+	 *
+	 *   Note that this wheel event sequence can happen *during* a 
+	 *   middle-mouse tracking sequence on a wheel-button (nice!)
+	 */
+
+
+	// This is not just called for keyboard keys, it is also called for mouse down/up events!
+	// Eg for left-press we get Key = EKeys::LeftMouseButton and Event = IE_Pressed
+	// Return value indicates "handled'. If we return true for mouse press events then
+	// the StartTracking/CapturedMouseMove/EndTracking sequence is *not* called (but we
+	// also don't get MouseMove(), so the mouse-movement events appear to be lost?)
+	virtual bool InputKey(FEditorViewportClient* ViewportClient, FViewport* Viewport, FKey Key, EInputEvent Event) override;
+
+
+	// Called for 1D Axis movements - EKeys::MouseX, EKeys::MouseY, and EKeys::MouseWheelAxis
+	// Called even if we return true from InputKey which otherwise blocks tracking. 
+	// Return value indicates "handled" but has no effect on CapturedMouseMove events (ie we still get them)
+	virtual bool InputAxis(FEditorViewportClient* InViewportClient, FViewport* Viewport, int32 ControllerId, FKey Key, float Delta, float DeltaTime);
+
+
+	// called on mouse-down. return value indicates whether this was "handled" but
+	// does not mean we get exclusive capture events!
+	virtual bool StartTracking(FEditorViewportClient* InViewportClient, FViewport* InViewport) override;
+
+	// called during mouse-down mouse-move. Always called, return value is used to indicate whether
+	// we "handled" this mouse move (see FEditorModeTools::CapturedMouseMove) but does not pre-empt
+	// other editor modes from seeing this move event...
+	virtual bool CapturedMouseMove(FEditorViewportClient* InViewportClient, FViewport* InViewport, int32 InMouseX, int32 InMouseY) override;
+
+	// always called on mouse-up
+	virtual bool EndTracking(FEditorViewportClient* InViewportClient, FViewport* InViewport) override;
+
+
+	//////////////////
+	// End of FEdMode interface
+	//////////////////
+
+
+
+public:
+	DECLARE_MULTICAST_DELEGATE_OneParam(FOnModelingModeToolNotification, const FText&);
+	FOnModelingModeToolNotification OnToolNotificationMessage;
+	FOnModelingModeToolNotification OnToolWarningMessage;
+
+
+public:
+	virtual UEdModeInteractiveToolsContext* GetToolsContext() const
+	{
+		return ToolsContext;
+	}
+
+	virtual UInteractiveToolManager* GetToolManager() const
+	{ 
+		return ToolsContext->ToolManager;
+	}
+
+protected:
+
+	UEdModeInteractiveToolsContext* ToolsContext;
+
+	/** Command list lives here so that the key bindings on the commands can be processed in the viewport. */
+	TSharedPtr<FUICommandList> UICommandList;
+
+
+public:
+	/** Cached pointer to the viewport world interaction object we're using to interact with mesh elements */
+	class UViewportWorldInteraction* ViewportWorldInteraction;
+
+};
diff --git a/Engine/Plugins/Experimental/ModelingToolsEditorMode/Source/ModelingToolsEditorMode/Public/ModelingToolsEditorModeStyle.h b/Engine/Plugins/Experimental/ModelingToolsEditorMode/Source/ModelingToolsEditorMode/Public/ModelingToolsEditorModeStyle.h
new file mode 100644
--- /dev/null
+++ b/Engine/Plugins/Experimental/ModelingToolsEditorMode/Source/ModelingToolsEditorMode/Public/ModelingToolsEditorModeStyle.h
@@ -0,0 +1,24 @@
+// Copyright 1998-2019 Epic Games, Inc. All Rights Reserved.
+
+#pragma once
+
+#include "CoreMinimal.h"
+#include "Styling/ISlateStyle.h"
+#include "Styling/SlateStyle.h"
+
+class FModelingToolsEditorModeStyle
+{
+public:
+	static void Initialize();
+
+	static void Shutdown();
+
+	static TSharedPtr< class ISlateStyle > Get();
+
+	static FName GetStyleSetName();
+private:
+	static FString InContent(const FString& RelativePath, const ANSICHAR* Extension);
+
+private:
+	static TSharedPtr< class FSlateStyleSet > StyleSet;
+};
diff --git a/Engine/Plugins/Experimental/ProxyLODPlugin/Source/ProxyLOD/Private/ProxyLODMaterialTransferUtilities.cpp b/Engine/Plugins/Experimental/ProxyLODPlugin/Source/ProxyLOD/Private/ProxyLODMaterialTransferUtilities.cpp
--- a/Engine/Plugins/Experimental/ProxyLODPlugin/Source/ProxyLOD/Private/ProxyLODMaterialTransferUtilities.cpp
+++ b/Engine/Plugins/Experimental/ProxyLODPlugin/Source/ProxyLOD/Private/ProxyLODMaterialTransferUtilities.cpp
@@ -9,9 +9,7 @@
 #include "ProxyLODMeshConvertUtils.h"
 #include "Math/Matrix.h" // used in chirality..
 
-#include "MeshDescription.h"
-#include "MeshAttributes.h"
-#include "MeshAttributeArray.h"
+#include "StaticMeshAttributes.h"
 
 
 namespace ProxyLOD
diff --git a/Engine/Plugins/Experimental/ProxyLODPlugin/Source/ProxyLOD/Private/ProxyLODMeshSDFConversions.cpp b/Engine/Plugins/Experimental/ProxyLODPlugin/Source/ProxyLOD/Private/ProxyLODMeshSDFConversions.cpp
--- a/Engine/Plugins/Experimental/ProxyLODPlugin/Source/ProxyLOD/Private/ProxyLODMeshSDFConversions.cpp
+++ b/Engine/Plugins/Experimental/ProxyLODPlugin/Source/ProxyLOD/Private/ProxyLODMeshSDFConversions.cpp
@@ -122,7 +122,7 @@ static openvdb::FloatGrid::Ptr OffsetSDF(const openvdb::FloatGrid::Ptr InSDFVolu
 
 	// convert the mesh to FMeshDescription
 	FMeshDescription RawMesh;
-	UStaticMesh::RegisterMeshAttributes(RawMesh);
+	FStaticMeshAttributes(RawMesh).Register();
 	ProxyLOD::MixedPolyMeshToRawMesh(MixedPolyMesh, RawMesh);
 
 	// Create a new empty grid with the same transform and metadata
diff --git a/Engine/Plugins/Experimental/ProxyLODPlugin/Source/ProxyLOD/Private/ProxyLODMeshTypes.h b/Engine/Plugins/Experimental/ProxyLODPlugin/Source/ProxyLOD/Private/ProxyLODMeshTypes.h
--- a/Engine/Plugins/Experimental/ProxyLODPlugin/Source/ProxyLOD/Private/ProxyLODMeshTypes.h
+++ b/Engine/Plugins/Experimental/ProxyLODPlugin/Source/ProxyLOD/Private/ProxyLODMeshTypes.h
@@ -11,9 +11,7 @@
 #endif
 
 #include "MeshMergeData.h"
-#include "MeshDescription.h"
-#include "MeshAttributes.h"
-#include "MeshAttributeArray.h"
+#include "StaticMeshAttributes.h"
 #include "MeshDescriptionOperations.h"
 
 #include "ProxyLODThreadedWrappers.h"
diff --git a/Engine/Plugins/Experimental/ProxyLODPlugin/Source/ProxyLOD/Private/ProxyLODMeshUtilities.cpp b/Engine/Plugins/Experimental/ProxyLODPlugin/Source/ProxyLOD/Private/ProxyLODMeshUtilities.cpp
--- a/Engine/Plugins/Experimental/ProxyLODPlugin/Source/ProxyLOD/Private/ProxyLODMeshUtilities.cpp
+++ b/Engine/Plugins/Experimental/ProxyLODPlugin/Source/ProxyLOD/Private/ProxyLODMeshUtilities.cpp
@@ -14,10 +14,7 @@ THIRD_PARTY_INCLUDES_START
 #include <DirectXMeshCode/DirectXMesh/DirectXMesh.h>
 THIRD_PARTY_INCLUDES_END
 
-#include "MeshDescription.h"
-#include "MeshAttributes.h"
-#include "MeshAttributeArray.h"
-
+#include "StaticMeshAttributes.h"
 #include "MeshDescriptionOperations.h"
 
 #include <vector>
diff --git a/Engine/Plugins/Experimental/ProxyLODPlugin/Source/ProxyLOD/Private/ProxyLODParameterization.cpp b/Engine/Plugins/Experimental/ProxyLODPlugin/Source/ProxyLOD/Private/ProxyLODParameterization.cpp
--- a/Engine/Plugins/Experimental/ProxyLODPlugin/Source/ProxyLOD/Private/ProxyLODParameterization.cpp
+++ b/Engine/Plugins/Experimental/ProxyLODPlugin/Source/ProxyLOD/Private/ProxyLODParameterization.cpp
@@ -5,6 +5,7 @@
 #include "ProxyLODMeshConvertUtils.h"
 #include "ProxyLODMeshUtilities.h"
 #include "ProxyLODMeshParameterization.h"
+#include "StaticMeshAttributes.h"
 
 
 
@@ -65,7 +66,7 @@ public:
 			}
 			// Reset the result mesh description
 			MeshDescription.Empty();
-			UStaticMesh::RegisterMeshAttributes(MeshDescription);
+			FStaticMeshAttributes(MeshDescription).Register();
 
 			// convert to fill the mesh description
 			ProxyLOD::ConvertMesh(VertexDataMesh, MeshDescription);
diff --git a/Engine/Plugins/Experimental/ProxyLODPlugin/Source/ProxyLOD/Private/ProxyLODPlugin.cpp b/Engine/Plugins/Experimental/ProxyLODPlugin/Source/ProxyLOD/Private/ProxyLODPlugin.cpp
--- a/Engine/Plugins/Experimental/ProxyLODPlugin/Source/ProxyLOD/Private/ProxyLODPlugin.cpp
+++ b/Engine/Plugins/Experimental/ProxyLODPlugin/Source/ProxyLOD/Private/ProxyLODPlugin.cpp
@@ -37,9 +37,7 @@ THIRD_PARTY_INCLUDES_END
 #include "ProxyLODkDOPInterface.h"
 #include "ProxyLODThreadedWrappers.h"
 
-#include "MeshDescription.h"
-#include "MeshAttributes.h"
-#include "MeshAttributeArray.h"
+#include "StaticMeshAttributes.h"
 #include "MeshDescriptionOperations.h"
 
 
@@ -601,7 +599,7 @@ void FVoxelizeMeshMerging::ProxyLOD(const FMeshMergeDataArray& InData, const FMe
 	//     will be captured by the CompleteDelegate. 
 
 	FMeshDescription OutRawMesh;
-	UStaticMesh::RegisterMeshAttributes(OutRawMesh);
+	FStaticMeshAttributes(OutRawMesh).Register();
 
 	FFlattenMaterial OutMaterial         = FMaterialUtilities::CreateFlattenMaterialWithSettings(InProxySettings.MaterialSettings);
 	const FColor UnresolvedGeometryColor = InProxySettings.UnresolvedGeometryColor;
diff --git a/Engine/Plugins/Experimental/ProxyLODPlugin/Source/ProxyLOD/Private/ProxyLODVolume.cpp b/Engine/Plugins/Experimental/ProxyLODPlugin/Source/ProxyLOD/Private/ProxyLODVolume.cpp
--- a/Engine/Plugins/Experimental/ProxyLODPlugin/Source/ProxyLOD/Private/ProxyLODVolume.cpp
+++ b/Engine/Plugins/Experimental/ProxyLODPlugin/Source/ProxyLOD/Private/ProxyLODVolume.cpp
@@ -99,7 +99,7 @@ public:
 		using namespace openvdb::tools;
 
 		FMeshDescription RawMesh;
-		UStaticMesh::RegisterMeshAttributes(RawMesh);
+		FStaticMeshAttributes(RawMesh).Register();
 		ConvertToRawMesh(RawMesh);
 		FMeshDescriptionAdapter MeshAdapter(RawMesh, SDFVolume->transform());
 
diff --git a/Engine/Plugins/Experimental/ProxyLODPlugin/Source/ProxyLOD/Private/ProxyLODkDOPUtils.cpp b/Engine/Plugins/Experimental/ProxyLODPlugin/Source/ProxyLOD/Private/ProxyLODkDOPUtils.cpp
--- a/Engine/Plugins/Experimental/ProxyLODPlugin/Source/ProxyLOD/Private/ProxyLODkDOPUtils.cpp
+++ b/Engine/Plugins/Experimental/ProxyLODPlugin/Source/ProxyLOD/Private/ProxyLODkDOPUtils.cpp
@@ -5,9 +5,7 @@
 #include "ProxyLODMeshTypes.h"
 #include "ProxyLODThreadedWrappers.h"
 
-#include "MeshAttributes.h"
-#include "MeshAttributeArray.h"
-#include "MeshDescription.h"
+#include "StaticMeshAttributes.h"
 
 // Utils for building a kdop tree from different mesh types.
 
diff --git a/Engine/Plugins/Experimental/SkeletalReduction/Source/Private/SkeletalMeshReductionPlugin.cpp b/Engine/Plugins/Experimental/SkeletalReduction/Source/Private/SkeletalMeshReductionPlugin.cpp
--- a/Engine/Plugins/Experimental/SkeletalReduction/Source/Private/SkeletalMeshReductionPlugin.cpp
+++ b/Engine/Plugins/Experimental/SkeletalReduction/Source/Private/SkeletalMeshReductionPlugin.cpp
@@ -100,8 +100,7 @@ public:
 	* @returns true if reduction was successful.
 	*/
 	virtual bool ReduceSkeletalMesh( class USkeletalMesh* SkeletalMesh,
-		                             int32 LODIndex,
-		                             bool bReregisterComponent = true
+		                             int32 LODIndex		                             
 	                                ) override ;
 
 
@@ -170,7 +169,8 @@ private:
 		                         FSkeletalMeshOptimizationSettings Settings,
 								 const FImportantBones& ImportantBones,
 		                         const TArray<FMatrix>& BoneMatrices,
-		                         const int32 LODIndex ) const;
+		                         const int32 LODIndex,
+								 const bool bReducingSourceModel) const;
 
 	/**
 	* Remove the specified section from the mesh.
@@ -204,7 +204,8 @@ private:
 		                                 const FSkeletalMeshLODModel& SrcLODModel,
 		                                 const SkeletalSimplifier::FSkinnedSkeletalMesh& SkinnedMesh,
 		                                 const FReferenceSkeleton& RefSkeleton,
-		                                 FSkeletalMeshLODModel& NewModel ) const;
+		                                 FSkeletalMeshLODModel& NewModel,
+										 const bool bReducingSourceModel) const;
 
 	/**
 	* Add the SourceModelInflunces to the LODModel in the case that alternate weights exists.
@@ -359,10 +360,7 @@ void FSkeletalMeshReduction::ShutdownModule()
 
 }
 
-
-
-
-bool FQuadricSkeletalMeshReduction::ReduceSkeletalMesh( USkeletalMesh* SkeletalMesh, int32 LODIndex, bool bReregisterComponent )
+bool FQuadricSkeletalMeshReduction::ReduceSkeletalMesh( USkeletalMesh* SkeletalMesh, int32 LODIndex )
 {
 	check(SkeletalMesh);
 	check(LODIndex >= 0);
@@ -372,21 +370,8 @@ bool FQuadricSkeletalMeshReduction::ReduceSkeletalMesh( USkeletalMesh* SkeletalM
 	check(SkeletalMeshResource);
 	check(LODIndex <= SkeletalMeshResource->LODModels.Num());
 
-	if (bReregisterComponent)
-	{
-		TComponentReregisterContext<USkinnedMeshComponent> ReregisterContext;
-		SkeletalMesh->ReleaseResources();
-		SkeletalMesh->ReleaseResourcesFence.Wait();
-
-		ReduceSkeletalMesh(*SkeletalMesh, *SkeletalMeshResource, LODIndex);
-
-		SkeletalMesh->PostEditChange();
-		SkeletalMesh->InitResources();
-	}
-	else
-	{
-		ReduceSkeletalMesh(*SkeletalMesh, *SkeletalMeshResource, LODIndex);
-	}
+	FScopedSkeletalMeshPostEditChange ScopedPostEditChange(SkeletalMesh);
+	ReduceSkeletalMesh(*SkeletalMesh, *SkeletalMeshResource, LODIndex);
 
 	return true;
 }
@@ -1541,8 +1526,8 @@ void FQuadricSkeletalMeshReduction::ConvertToFSkeletalMeshLODModel( const int32
 	                                                                const FSkeletalMeshLODModel& SrcLODModel,
 	                                                                const SkeletalSimplifier::FSkinnedSkeletalMesh& SkinnedMesh,
 	                                                                const FReferenceSkeleton& RefSkeleton,
-	                                                                FSkeletalMeshLODModel& NewModel
-                                                                   ) const
+	                                                                FSkeletalMeshLODModel& NewModel,	
+																	const bool bReducingSourceModel) const
 {
 	// We might be re-using this model - so clear it.
 
@@ -1553,11 +1538,14 @@ void FQuadricSkeletalMeshReduction::ConvertToFSkeletalMeshLODModel( const int32
 	FSkeletalMeshData SkeletalMeshData;
 	ExtractFSkeletalData(SkinnedMesh, SkeletalMeshData);
 
-	// Add alternate weight data to the NewModel.  Note, this has to be done before we "BuildSkeletalMesh" to insure
-	// that the bone-based vertex chunking respects the alternate weights.
-	// NB: this only prepares the NewModel, but BuildSkeletalMesh is only 1/2-aware of this, so we will have to do some additional work after.
-	
-	AddSourceModelInfluences(SrcLODModel, SkinnedMesh, NewModel);
+	//if (!bReducingSourceModel)
+	{
+		// Add alternate weight data to the NewModel.  Note, this has to be done before we "BuildSkeletalMesh" to insure
+		// that the bone-based vertex chunking respects the alternate weights.
+		// NB: this only prepares the NewModel, but BuildSkeletalMesh is only 1/2-aware of this, so we will have to do some additional work after.
+
+		AddSourceModelInfluences(SrcLODModel, SkinnedMesh, NewModel);
+	}
 	
 
 
@@ -1576,6 +1564,7 @@ void FQuadricSkeletalMeshReduction::ConvertToFSkeletalMeshLODModel( const int32
 	Options.bComputeNormals = false;
 	Options.bComputeTangents = false;
 	Options.bUseMikkTSpace = true; //Avoid builtin build by specifying true for mikkt space
+	Options.bComputeWeightedNormals = false;
 	Options.OverlappingThresholds.ThresholdPosition = 0.0f;
 	Options.OverlappingThresholds.ThresholdTangentNormal = 0.0f;
 	Options.OverlappingThresholds.ThresholdUV = 0.0f;
@@ -1595,15 +1584,19 @@ void FQuadricSkeletalMeshReduction::ConvertToFSkeletalMeshLODModel( const int32
 		Options
 	);
 
+	//Re-Apply the user section changes, the UserSectionsData is map to original section and should match the builded LODModel
+	NewModel.SyncronizeUserSectionsDataArray();
+
 	// Set texture coordinate count on the new model.
 	NewModel.NumTexCoords = SkeletalMeshData.TexCoordCount;
 
 
+	//if (!bReducingSourceModel)
+	{
+		// Update the alternate weights
 
-	// Update the alternate weights
-	
-	UpdateAlternateWeights(MaxBonesPerVertex, NewModel);
-
+		UpdateAlternateWeights(MaxBonesPerVertex, NewModel);
+	}
 }
 
 bool FQuadricSkeletalMeshReduction::ReduceSkeletalLODModel( const FSkeletalMeshLODModel& SrcModel,
@@ -1613,7 +1606,8 @@ bool FQuadricSkeletalMeshReduction::ReduceSkeletalLODModel( const FSkeletalMeshL
 	                                                        FSkeletalMeshOptimizationSettings Settings,
 															const FImportantBones& ImportantBones,
 	                                                        const TArray<FMatrix>& BoneMatrices,
-	                                                        const int32 LODIndex
+	                                                        const int32 LODIndex,
+															const bool bReducingSourceModel
                                                            ) const
 {
 
@@ -1682,7 +1676,7 @@ bool FQuadricSkeletalMeshReduction::ReduceSkeletalLODModel( const FSkeletalMeshL
 
 		// Convert to SkeletalMeshLODModel. 
 
-		ConvertToFSkeletalMeshLODModel(Settings.MaxBonesPerVertex, SrcModel, SkinnedSkeletalMesh, RefSkeleton, OutSkeletalMeshLODModel);
+		ConvertToFSkeletalMeshLODModel(Settings.MaxBonesPerVertex, SrcModel, SkinnedSkeletalMesh, RefSkeleton, OutSkeletalMeshLODModel, bReducingSourceModel);
 
 		// We may need to do additional simplification if the user specified a hard number limit for verts and
 		// the internal chunking during conversion split some verts.
@@ -1804,13 +1798,26 @@ void FQuadricSkeletalMeshReduction::ReduceSkeletalMesh(USkeletalMesh& SkeletalMe
 	struct FSectionData
 	{
 		uint16 MaterialIndex;
+		int32 MaterialMap;
 		bool bCastShadow;
 		bool bRecomputeTangent;
+		bool bDisabled;
+		int32 GenerateUpToLodIndex;
+		int32 ChunkedParentSectionIndex;
+		int32 OriginalDataSectionIndex;
 	};
 
 	TMap<int32, FSectionData> BackupSectionIndexToSectionData;
+	TMap<int32, FSkelMeshSourceSectionUserData> BackupUserSectionsData;
+	FString BackupLodModelBuildStringID = TEXT("");
 	//Store the section data. Store the source LOD model in case we add a LOD, or the target LOD model in case the LOD already exist
 	FSkeletalMeshLODModel& BackupSectionLODModel = bLODModelAdded ? SkeletalMeshResource.LODModels[Settings.BaseLOD] : SkeletalMeshResource.LODModels[LODIndex];
+	if (!bLODModelAdded)
+	{
+		BackupLodModelBuildStringID = BackupSectionLODModel.BuildStringID;
+		BackupUserSectionsData = BackupSectionLODModel.UserSectionsData;
+	}
+
 	BackupSectionIndexToSectionData.Reserve(BackupSectionLODModel.Sections.Num());
 
 	for (int32 SectionIndex = 0; SectionIndex < BackupSectionLODModel.Sections.Num(); ++SectionIndex)
@@ -1825,23 +1832,27 @@ void FQuadricSkeletalMeshReduction::ReduceSkeletalMesh(USkeletalMesh& SkeletalMe
 		SectionData.MaterialIndex = BackupSectionLODModel.Sections[SectionIndex].MaterialIndex;
 		SectionData.bCastShadow = BackupSectionLODModel.Sections[SectionIndex].bCastShadow;
 		SectionData.bRecomputeTangent = BackupSectionLODModel.Sections[SectionIndex].bRecomputeTangent;
+		SectionData.bDisabled = BackupSectionLODModel.Sections[SectionIndex].bDisabled;
+		SectionData.GenerateUpToLodIndex = BackupSectionLODModel.Sections[SectionIndex].GenerateUpToLodIndex;
+		SectionData.ChunkedParentSectionIndex = BackupSectionLODModel.Sections[SectionIndex].ChunkedParentSectionIndex;
+		SectionData.OriginalDataSectionIndex = BackupSectionLODModel.Sections[SectionIndex].OriginalDataSectionIndex;
+		SectionData.MaterialMap = LODInfo->LODMaterialMap.IsValidIndex(SectionIndex) ? LODInfo->LODMaterialMap[SectionIndex] : INDEX_NONE;
+		if (SectionData.MaterialMap == SectionData.MaterialIndex)
+		{
+			//Remove any override if the value is the same
+			SectionData.MaterialMap = INDEX_NONE;
+		}
 	}
 
-	auto FillClothingData = [&SkeletalMeshResource, &LODIndex, bLODModelAdded](int32 &EnableSectionNumber, TArray<bool> &SectionStatus)
+	auto FillSectionMaterialSlot = [&SkeletalMeshResource, &LODIndex, bLODModelAdded](TArray<int32>& SectionMaterialSlot)
 	{
-		EnableSectionNumber = 0;
-		SectionStatus.Empty();
+		SectionMaterialSlot.Empty();
 		if (!bLODModelAdded && SkeletalMeshResource.LODModels.IsValidIndex(LODIndex))
 		{
 			int32 SectionNumber = SkeletalMeshResource.LODModels[LODIndex].Sections.Num();
-			SectionStatus.Reserve(SectionNumber);
 			for (int32 SectionIndex = 0; SectionIndex < SectionNumber; ++SectionIndex)
 			{
-				SectionStatus.Add(!SkeletalMeshResource.LODModels[LODIndex].Sections[SectionIndex].bDisabled);
-				if (SectionStatus[SectionIndex])
-				{
-					EnableSectionNumber++;
-				}
+				SectionMaterialSlot.Add(SkeletalMeshResource.LODModels[LODIndex].Sections[SectionIndex].OriginalDataSectionIndex);
 			}
 		}
 	};
@@ -1849,31 +1860,20 @@ void FQuadricSkeletalMeshReduction::ReduceSkeletalMesh(USkeletalMesh& SkeletalMe
 	// Unbind any existing clothing assets before we reimport the geometry
 	TArray<ClothingAssetUtils::FClothingAssetMeshBinding> ClothingBindings;
 	//Get a map of enable/disable sections
-	int32 OriginalSectionNumberBeforeReduction = 0;
-	TArray<bool> OriginalSectionEnableBeforeReduction;
+	TArray<int32> OriginalSectionMaterialSlot;
 
 	//Do not play with cloth if the LOD is added
 	if (!bLODModelAdded)
 	{
-		//Store the clothBinding
-		ClothingAssetUtils::GetMeshClothingAssetBindings(&SkeletalMesh, ClothingBindings, LODIndex);
-		FillClothingData(OriginalSectionNumberBeforeReduction, OriginalSectionEnableBeforeReduction);
-		//Unbind the Cloth for this LOD before we reduce it, we will put back the cloth after the reduction, if it still match the sections
-		for (ClothingAssetUtils::FClothingAssetMeshBinding& Binding : ClothingBindings)
-		{
-			if (Binding.LODIndex == LODIndex)
-			{
-				Binding.Asset->UnbindFromSkeletalMesh(&SkeletalMesh, Binding.LODIndex);
-			}
-		}
+		FLODUtilities::UnbindClothingAndBackup(&SkeletalMesh, ClothingBindings, LODIndex);
 	}
 
 	bool bReducingSourceModel = false;
-	//Reducing Base LOD, we need to use the temporary data so it can be iterative
-	if (BaseLOD == LODIndex && SkelResource->OriginalReductionSourceMeshData.IsValidIndex(BaseLOD) && !SkelResource->OriginalReductionSourceMeshData[BaseLOD]->IsEmpty())
+	//Reducing source LOD, we need to use the temporary data so it can be iterative
+	if (BaseLOD == LODIndex && (SkelResource->OriginalReductionSourceMeshData.IsValidIndex(BaseLOD) && !SkelResource->OriginalReductionSourceMeshData[BaseLOD]->IsEmpty()))
 	{
 		TMap<FString, TArray<FMorphTargetDelta>> TempLODMorphTargetData;
-		SkelResource->OriginalReductionSourceMeshData[BaseLOD]->LoadReductionData(*SrcModel, TempLODMorphTargetData);
+		SkelResource->OriginalReductionSourceMeshData[BaseLOD]->LoadReductionData(*SrcModel, TempLODMorphTargetData, &SkeletalMesh);
 		bReducingSourceModel = true;
 	}
 	else
@@ -1955,6 +1955,8 @@ void FQuadricSkeletalMeshReduction::ReduceSkeletalMesh(USkeletalMesh& SkeletalMe
 
 	// Swap out the old model.  
 	FSkeletalMeshImportData RawMesh;
+	ESkeletalMeshGeoImportVersions GeoImportVersion = ESkeletalMeshGeoImportVersions::Before_Versionning;
+	ESkeletalMeshSkinningImportVersions SkinningImportVersion = ESkeletalMeshSkinningImportVersions::Before_Versionning;
 	{
 		FSkeletalMeshLODModel* Old = LODModels[LODIndex];
 
@@ -1962,12 +1964,21 @@ void FQuadricSkeletalMeshReduction::ReduceSkeletalMesh(USkeletalMesh& SkeletalMe
 
 		if (!bReducingSourceModel && Old)
 		{
+			//We need to backup the original RawSkeletalMeshBulkData in case it was an imported LOD
+			if (!bLODModelAdded && !Old->RawSkeletalMeshBulkData.IsEmpty())
+			{
+				Old->RawSkeletalMeshBulkData.LoadRawMesh(RawMesh);
+				GeoImportVersion = Old->RawSkeletalMeshBulkData.GeoImportVersion;
+				SkinningImportVersion = Old->RawSkeletalMeshBulkData.SkinningImportVersion;
+			}
 			delete Old;
 		}
 		else if(bReducingSourceModel)
 		{
 			//In case we reduce the source model we want to keep the original import data
 			SrcModel->RawSkeletalMeshBulkData.LoadRawMesh(RawMesh);
+			GeoImportVersion = SrcModel->RawSkeletalMeshBulkData.GeoImportVersion;
+			SkinningImportVersion = SrcModel->RawSkeletalMeshBulkData.SkinningImportVersion;
 		}
 	}
 
@@ -1975,8 +1986,10 @@ void FQuadricSkeletalMeshReduction::ReduceSkeletalMesh(USkeletalMesh& SkeletalMe
 
 	// Reduce LOD model with SrcMesh
 
-	if (ReduceSkeletalLODModel(*SrcModel, *NewModel, SkeletalMesh.GetImportedBounds(), SkeletalMesh.RefSkeleton, Settings, ImportantBones, RelativeToRefPoseMatrices, LODIndex))
+	if (ReduceSkeletalLODModel(*SrcModel, *NewModel, SkeletalMesh.GetImportedBounds(), SkeletalMesh.RefSkeleton, Settings, ImportantBones, RelativeToRefPoseMatrices, LODIndex, bReducingSourceModel))
 	{
+		FSkeletalMeshLODInfo* ReducedLODInfoPtr = SkeletalMesh.GetLODInfo(LODIndex);
+		check(ReducedLODInfoPtr);
 		// Do any joint-welding / bone removal.
 
 		if (MeshBoneReductionInterface != NULL && MeshBoneReductionInterface->GetBoneReductionData(&SkeletalMesh, LODIndex, BonesToRemove))
@@ -1990,35 +2003,22 @@ void FQuadricSkeletalMeshReduction::ReduceSkeletalMesh(USkeletalMesh& SkeletalMe
 
 		if (bOldLodWasFromFile)
 		{
-			SkeletalMesh.GetLODInfo(LODIndex)->LODMaterialMap.Empty();
+			ReducedLODInfoPtr->LODMaterialMap.Empty();
 		}
-
-		// If base lod has a customized LODMaterialMap and this LOD doesn't (could have if changes are applied instead of freshly generated, copy over the data into new new LOD
-
-		if (SkeletalMesh.GetLODInfo(LODIndex)->LODMaterialMap.Num() == 0 && SkeletalMesh.GetLODInfo(BaseLOD)->LODMaterialMap.Num() != 0)
-		{
-			SkeletalMesh.GetLODInfo(LODIndex)->LODMaterialMap = SkeletalMesh.GetLODInfo(BaseLOD)->LODMaterialMap;
-		}
-		else
-		{
-			// Assuming the reducing step has set all material indices correctly, we double check if something went wrong
-			// make sure we don't have more materials
-			int32 TotalSectionCount = NewModel->Sections.Num();
-			if (SkeletalMesh.GetLODInfo(LODIndex)->LODMaterialMap.Num() > TotalSectionCount)
-			{
-				SkeletalMesh.GetLODInfo(LODIndex)->LODMaterialMap = SkeletalMesh.GetLODInfo(BaseLOD)->LODMaterialMap;
-				// Something went wrong during the reduce step during regenerate 					
-				check(SkeletalMesh.GetLODInfo(BaseLOD)->LODMaterialMap.Num() == TotalSectionCount || SkeletalMesh.GetLODInfo(BaseLOD)->LODMaterialMap.Num() == 0);
-			}
-		}
-
 		// Flag this LOD as having been simplified.
-		SkeletalMesh.GetLODInfo(LODIndex)->bHasBeenSimplified = true;
+		ReducedLODInfoPtr->bHasBeenSimplified = true;
 		SkeletalMesh.bHasBeenSimplified = true;
 		//Restore section data
 		FSkeletalMeshLODModel& ImportedModelLOD = SkeletalMesh.GetImportedModel()->LODModels[LODIndex];
 		TArray<bool> SectionMatched;
 		SectionMatched.AddZeroed(ImportedModelLOD.Sections.Num());
+		
+		if (!bLODModelAdded)
+		{
+			ImportedModelLOD.UserSectionsData = BackupUserSectionsData;
+			ImportedModelLOD.BuildStringID = BackupLodModelBuildStringID;
+		}
+
 		for (auto Kvp : BackupSectionIndexToSectionData)
 		{
 			const int32 SourceSectionIndex = Kvp.Key;
@@ -2031,10 +2031,24 @@ void FQuadricSkeletalMeshReduction::ReduceSkeletalMesh(USkeletalMesh& SkeletalMe
 				{
 					continue;
 				}
-				if (SectionData.MaterialIndex == ImportedModelLOD.Sections[SectionIndex].MaterialIndex)
+				if (SectionData.OriginalDataSectionIndex == ImportedModelLOD.Sections[SectionIndex].OriginalDataSectionIndex)
 				{
+					//Restore the LOD material map for this section
+					if (SectionData.MaterialMap != INDEX_NONE)
+					{
+						while (ReducedLODInfoPtr->LODMaterialMap.Num() <= SectionIndex)
+						{
+							ReducedLODInfoPtr->LODMaterialMap.Add(INDEX_NONE);
+						}
+						ReducedLODInfoPtr->LODMaterialMap[SectionIndex] = SectionData.MaterialMap;
+					}
+					//Restore the LODModel Section
 					ImportedModelLOD.Sections[SectionIndex].bCastShadow = SectionData.bCastShadow;
 					ImportedModelLOD.Sections[SectionIndex].bRecomputeTangent = SectionData.bRecomputeTangent;
+					ImportedModelLOD.Sections[SectionIndex].bDisabled = SectionData.bDisabled;
+					ImportedModelLOD.Sections[SectionIndex].GenerateUpToLodIndex = SectionData.GenerateUpToLodIndex;
+					ImportedModelLOD.Sections[SectionIndex].ChunkedParentSectionIndex = SectionData.ChunkedParentSectionIndex;
+					ImportedModelLOD.Sections[SectionIndex].OriginalDataSectionIndex = SectionData.OriginalDataSectionIndex;
 					SectionMatched[SectionIndex] = true;
 					break;
 				}
@@ -2087,47 +2101,22 @@ void FQuadricSkeletalMeshReduction::ReduceSkeletalMesh(USkeletalMesh& SkeletalMe
 	if (!bLODModelAdded)
 	{
 		//Get the number of enabled section
-		int32 SectionNumberAfterReduction = 0;
-		TArray<bool> SectionEnableAfterReduction;
-		FillClothingData(SectionNumberAfterReduction, SectionEnableAfterReduction);
+		TArray<int32> SectionMaterialSlotAfterReduction;
+		FillSectionMaterialSlot(SectionMaterialSlotAfterReduction);
 
-		//Put back the clothing for this newly reduce LOD only if the section count match.
-		if (ClothingBindings.Num() > 0 && OriginalSectionNumberBeforeReduction == SectionNumberAfterReduction)
+		//Put back the clothing for this newly reduce LOD
+		if (ClothingBindings.Num() > 0)
 		{
-			TArray<int32> RemapSectionIndex;
-			int32 SectionIndexTest = 0;
-			for (int32 SectionIndexRef = 0; SectionIndexRef < OriginalSectionEnableBeforeReduction.Num(); SectionIndexRef++)
-			{
-				int32& RemapValue = RemapSectionIndex.Add_GetRef(INDEX_NONE);
-				if (!OriginalSectionEnableBeforeReduction[SectionIndexRef])
-				{
-					continue;
-				}
-				for (; SectionIndexTest <= SectionIndexRef; SectionIndexTest++)
-				{
-					if (SectionEnableAfterReduction.IsValidIndex(SectionIndexTest) && SectionEnableAfterReduction[SectionIndexTest])
-					{
-						RemapValue = SectionIndexTest++;
-						break;
-					}
-				}
-			}
-
-			for (ClothingAssetUtils::FClothingAssetMeshBinding& Binding : ClothingBindings)
-			{
-				int32 RemapBindingSectionIndex = RemapSectionIndex[Binding.SectionIndex];
-				if (RemapBindingSectionIndex != INDEX_NONE && Binding.LODIndex == LODIndex && NewModel->Sections.IsValidIndex(RemapBindingSectionIndex))
-				{
-					Binding.Asset->BindToSkeletalMesh(&SkeletalMesh, Binding.LODIndex, RemapBindingSectionIndex, Binding.AssetInternalLodIndex, false);
-				}
-			}
+			FLODUtilities::RestoreClothingFromBackup(&SkeletalMesh, ClothingBindings, LODIndex);
 		}
 	}
 
-	if (bReducingSourceModel && RawMesh.Points.Num() > 0)
+	if ((bReducingSourceModel || !bLODModelAdded ) && RawMesh.Points.Num() > 0)
 	{
 		//Put back the original import data, we need it to allow inline reduction and skeletal mesh split workflow
 		SkeletalMeshResource.LODModels[LODIndex].RawSkeletalMeshBulkData.SaveRawMesh(RawMesh);
+		SkeletalMeshResource.LODModels[LODIndex].RawSkeletalMeshBulkData.GeoImportVersion = GeoImportVersion;
+		SkeletalMeshResource.LODModels[LODIndex].RawSkeletalMeshBulkData.SkinningImportVersion = SkinningImportVersion;
 	}
 
 	SkeletalMesh.CalculateRequiredBones(SkeletalMeshResource.LODModels[LODIndex], SkeletalMesh.RefSkeleton, &BonesToRemove);
diff --git a/Engine/Plugins/Importers/USDImporter/Source/USDImporter/Private/StaticMeshImporter.cpp b/Engine/Plugins/Importers/USDImporter/Source/USDImporter/Private/StaticMeshImporter.cpp
--- a/Engine/Plugins/Importers/USDImporter/Source/USDImporter/Private/StaticMeshImporter.cpp
+++ b/Engine/Plugins/Importers/USDImporter/Source/USDImporter/Private/StaticMeshImporter.cpp
@@ -11,8 +11,7 @@
 #include "UObject/Package.h"
 #include "USDAssetImportData.h"
 #include "Factories/Factory.h"
-#include "MeshDescription.h"
-#include "MeshAttributes.h"
+#include "StaticMeshAttributes.h"
 #include "IMeshBuilderModule.h"
 #include "PackageTools.h"
 
@@ -23,40 +22,6 @@
 
 #define LOCTEXT_NAMESPACE "USDImportPlugin"
 
-struct FMeshDescriptionWrapper
-{
-	struct FVertexAttributes
-	{
-		FVertexAttributes(FMeshDescription* MeshDescription) :
-			Positions(MeshDescription->VertexAttributes().GetAttributesRef<FVector>(MeshAttribute::Vertex::Position)),
-			Normals(MeshDescription->VertexInstanceAttributes().GetAttributesRef<FVector>(MeshAttribute::VertexInstance::Normal)),
-			Tangents(MeshDescription->VertexInstanceAttributes().GetAttributesRef<FVector>(MeshAttribute::VertexInstance::Tangent)),
-			BinormalSigns(MeshDescription->VertexInstanceAttributes().GetAttributesRef<float>(MeshAttribute::VertexInstance::BinormalSign)),
-			Colors(MeshDescription->VertexInstanceAttributes().GetAttributesRef<FVector4>(MeshAttribute::VertexInstance::Color)),
-			UVs(MeshDescription->VertexInstanceAttributes().GetAttributesRef<FVector2D>(MeshAttribute::VertexInstance::TextureCoordinate))
-		{
-		}
-		TVertexAttributesRef<FVector> Positions;
-		TVertexInstanceAttributesRef<FVector> Normals;
-		TVertexInstanceAttributesRef<FVector> Tangents;
-		TVertexInstanceAttributesRef<float> BinormalSigns;
-		TVertexInstanceAttributesRef<FVector4> Colors;
-		TVertexInstanceAttributesRef<FVector2D> UVs;
-	};
-
-	FMeshDescriptionWrapper(FMeshDescription* MeshDescription) :
-		Vertex(MeshDescription),
-		EdgeHardnesses(MeshDescription->EdgeAttributes().GetAttributesRef<bool>(MeshAttribute::Edge::IsHard)),
-		EdgeCreaseSharpnesses(MeshDescription->EdgeAttributes().GetAttributesRef<float>(MeshAttribute::Edge::CreaseSharpness)),
-		PolygonGroupImportedMaterialSlotNames(MeshDescription->PolygonGroupAttributes().GetAttributesRef<FName>(MeshAttribute::PolygonGroup::ImportedMaterialSlotName))
-	{
-	}
-
-	FVertexAttributes Vertex;
-	TEdgeAttributesRef<bool> EdgeHardnesses;
-	TEdgeAttributesRef<float> EdgeCreaseSharpnesses;
-	TPolygonGroupAttributesRef<FName> PolygonGroupImportedMaterialSlotNames;
-};
 
 struct FUSDImportMaterialInfo
 {
@@ -99,8 +64,8 @@ public:
 	void ProcessMaterials(int32 LODIndex);
 
 private:
-	void AddVertexPositions( FMeshDescriptionWrapper& DestMeshWrapper, const pxr::UsdGeomMesh& Mesh );
-	bool AddPolygons( FMeshDescriptionWrapper& DestMeshWrapper, const pxr::UsdGeomMesh& Mesh, const TArray< FString >& MaterialNames, const TArray< int32 >& FaceMaterialIndices );
+	void AddVertexPositions( const pxr::UsdGeomMesh& Mesh );
+	bool AddPolygons( const pxr::UsdGeomMesh& Mesh, const TArray< FString >& MaterialNames, const TArray< int32 >& FaceMaterialIndices );
 };
 
 void FUSDStaticMeshImportState::ProcessStaticUSDGeometry(const pxr::UsdPrim& GeomPrim, int32 LODIndex)
@@ -111,7 +76,6 @@ void FUSDStaticMeshImportState::ProcessStaticUSDGeometry(const pxr::UsdPrim& Geo
 		return;
 	}
 
-	FMeshDescriptionWrapper DestMeshWrapper(MeshDescription);
 	VertexOffset = MeshDescription->Vertices().Num();
 	VertexInstanceOffset = MeshDescription->VertexInstances().Num();
 	PolygonOffset = MeshDescription->Polygons().Num();
@@ -121,14 +85,15 @@ void FUSDStaticMeshImportState::ProcessStaticUSDGeometry(const pxr::UsdPrim& Geo
 
 	Materials.AddDefaulted( GeometryMaterials.Key.Num() );
 
-	AddVertexPositions( DestMeshWrapper, Mesh );
-	AddPolygons( DestMeshWrapper, Mesh, GeometryMaterials.Key, GeometryMaterials.Value );
+	AddVertexPositions( Mesh );
+	AddPolygons( Mesh, GeometryMaterials.Key, GeometryMaterials.Value );
 }
 
-void FUSDStaticMeshImportState::AddVertexPositions( FMeshDescriptionWrapper& DestMeshWrapper, const pxr::UsdGeomMesh& Mesh )
+void FUSDStaticMeshImportState::AddVertexPositions( const pxr::UsdGeomMesh& Mesh )
 {
 	using namespace pxr;
 
+	FStaticMeshAttributes Attributes(*MeshDescription);
 	pxr::UsdAttribute Points = Mesh.GetPointsAttr();
 	if ( Points )
 	{
@@ -143,15 +108,16 @@ void FUSDStaticMeshImportState::AddVertexPositions( FMeshDescriptionWrapper& Des
 			Pos = FinalTransform.TransformPosition(Pos);
 
 			FVertexID AddedVertexId = MeshDescription->CreateVertex();
-			DestMeshWrapper.Vertex.Positions[AddedVertexId] = Pos;
+			Attributes.GetVertexPositions()[AddedVertexId] = Pos;
 		}
 	}
 }
 
-bool FUSDStaticMeshImportState::AddPolygons( FMeshDescriptionWrapper& DestMeshWrapper, const pxr::UsdGeomMesh& Mesh, const TArray< FString >& MaterialNames, const TArray< int32 >& FaceMaterialIndices )
+bool FUSDStaticMeshImportState::AddPolygons( const pxr::UsdGeomMesh& Mesh, const TArray< FString >& MaterialNames, const TArray< int32 >& FaceMaterialIndices )
 {
 	using namespace pxr;
 
+	FStaticMeshAttributes Attributes(*MeshDescription);
 	TMap<int32, FPolygonGroupID> PolygonGroupMapping;
 	TArray<FVertexInstanceID> CornerInstanceIDs;
 	TArray<FVertexID> CornerVerticesIDs;
@@ -232,14 +198,14 @@ bool FUSDStaticMeshImportState::AddPolygons( FMeshDescriptionWrapper& DestMeshWr
 
 	// When importing multiple mesh pieces to the same static mesh.  Ensure each mesh piece has the same number of Uv's
 	{
-		int32 ExistingUVCount = DestMeshWrapper.Vertex.UVs.GetNumIndices();
+		int32 ExistingUVCount = Attributes.GetVertexInstanceUVs().GetNumIndices();
 		int32 NumUVs = FMath::Max(UVSets.Num(), ExistingUVCount);
 		NumUVs = FMath::Min<int32>(MAX_MESH_TEXTURE_COORDS_MD, NumUVs);
 		// At least one UV set must exist.  
 		NumUVs = FMath::Max<int32>(1, NumUVs);
 
 		//Make sure all Vertex instance have the correct number of UVs
-		DestMeshWrapper.Vertex.UVs.SetNumIndices(NumUVs);
+		Attributes.GetVertexInstanceUVs().SetNumIndices(NumUVs);
 	}
 
 	for ( int32 PolygonIndex = 0; PolygonIndex < FaceCounts.size(); ++PolygonIndex )
@@ -257,7 +223,7 @@ bool FUSDStaticMeshImportState::AddPolygons( FMeshDescriptionWrapper& DestMeshWr
 			CornerInstanceIDs[CornerIndex] = VertexInstanceID;
 			const int32 ControlPointIndex = FaceIndices[CurrentVertexInstanceIndex];
 			const FVertexID VertexID(VertexOffset + ControlPointIndex);
-			const FVector VertexPosition = DestMeshWrapper.Vertex.Positions[VertexID];
+			const FVector VertexPosition = Attributes.GetVertexPositions()[VertexID];
 			CornerVerticesIDs[CornerIndex] = VertexID;
 
 			FVertexInstanceID AddedVertexInstanceId = MeshDescription->CreateVertexInstance(VertexID);
@@ -269,7 +235,7 @@ bool FUSDStaticMeshImportState::AddPolygons( FMeshDescriptionWrapper& DestMeshWr
 				const GfVec3f& Normal = Normals[NormalIndex];
 				FVector TransformedNormal = FinalTransformIT.TransformVector( USDToUnreal::ConvertVector( *ImportContext.Stage, Normal ) );
 
-				DestMeshWrapper.Vertex.Normals[AddedVertexInstanceId] = TransformedNormal.GetSafeNormal();
+				Attributes.GetVertexInstanceNormals()[AddedVertexInstanceId] = TransformedNormal.GetSafeNormal();
 			}
 
 			int32 UVLayerIndex = 0;
@@ -296,7 +262,7 @@ bool FUSDStaticMeshImportState::AddPolygons( FMeshDescriptionWrapper& DestMeshWr
 
 				// Flip V for Unreal uv's which match directx
 				FVector2D FinalUVVector(UV[0], 1.f - UV[1]);
-				DestMeshWrapper.Vertex.UVs.Set(AddedVertexInstanceId, UVLayerIndex, FinalUVVector);
+				Attributes.GetVertexInstanceUVs().Set(AddedVertexInstanceId, UVLayerIndex, FinalUVVector);
 
 				++UVLayerIndex;
 			}
@@ -326,7 +292,7 @@ bool FUSDStaticMeshImportState::AddPolygons( FMeshDescriptionWrapper& DestMeshWr
 			FPolygonGroupID ExistingPolygonGroup = FPolygonGroupID::Invalid;
 			for (const FPolygonGroupID PolygonGroupID : MeshDescription->PolygonGroups().GetElementIDs())
 			{
-				if (DestMeshWrapper.PolygonGroupImportedMaterialSlotNames[PolygonGroupID] == ImportedMaterialSlotName)
+				if (Attributes.GetPolygonGroupMaterialSlotNames()[PolygonGroupID] == ImportedMaterialSlotName)
 				{
 					ExistingPolygonGroup = PolygonGroupID;
 					break;
@@ -335,7 +301,7 @@ bool FUSDStaticMeshImportState::AddPolygons( FMeshDescriptionWrapper& DestMeshWr
 			if (ExistingPolygonGroup == FPolygonGroupID::Invalid)
 			{
 				ExistingPolygonGroup = MeshDescription->CreatePolygonGroup();
-				DestMeshWrapper.PolygonGroupImportedMaterialSlotNames[ExistingPolygonGroup] = ImportedMaterialSlotName;
+				Attributes.GetPolygonGroupMaterialSlotNames()[ExistingPolygonGroup] = ImportedMaterialSlotName;
 			}
 			PolygonGroupMapping.Add(RealMaterialIndex, ExistingPolygonGroup);
 		}
@@ -356,7 +322,7 @@ bool FUSDStaticMeshImportState::AddPolygons( FMeshDescriptionWrapper& DestMeshWr
 		pxr::VtArray<pxr::GfVec3f> USDColors;
 		ColorPrimvar.ComputeFlattened(&USDColors);
 
-		TVertexInstanceAttributesRef<FVector4> Colors = MeshDescription->VertexInstanceAttributes().GetAttributesRef<FVector4>(MeshAttribute::VertexInstance::Color);
+		TVertexInstanceAttributesRef<FVector4> Colors = Attributes.GetVertexInstanceColors();
 
 		int32 NumColors = USDColors.size();
 
@@ -397,11 +363,11 @@ void FUSDStaticMeshImportState::ProcessMaterials(int32 LODIndex)
 {
 	const FString BasePackageName = FPackageName::GetLongPackagePath(NewMesh->GetOutermost()->GetName());
 
-	FMeshDescriptionWrapper DestMeshWrapper(MeshDescription);
+	FStaticMeshAttributes Attributes(*MeshDescription);
 	TArray<FStaticMaterial> MaterialToAdd;
 	for (const FPolygonGroupID PolygonGroupID : MeshDescription->PolygonGroups().GetElementIDs())
 	{
-		const FName& ImportedMaterialSlotName = DestMeshWrapper.PolygonGroupImportedMaterialSlotNames[PolygonGroupID];
+		const FName& ImportedMaterialSlotName = Attributes.GetPolygonGroupMaterialSlotNames()[PolygonGroupID];
 		const FString ImportedMaterialSlotNameString = ImportedMaterialSlotName.ToString();
 		const FName MaterialSlotName = ImportedMaterialSlotName;
 		int32 MaterialIndex = INDEX_NONE;
@@ -550,7 +516,6 @@ UStaticMesh* FUSDStaticMeshImporter::ImportStaticMesh(FUsdImportContext& ImportC
 		//Create private asset in the same package as the StaticMesh, and make sure reference are set to avoid GC
 		State.MeshDescription = NewMesh->CreateMeshDescription(LODIndex);
 		check(State.MeshDescription != nullptr);
-		NewMesh->RegisterMeshAttributes(*State.MeshDescription);
 
 		bool bRecomputeNormals = false;
 
diff --git a/Engine/Plugins/Runtime/EditableMesh/Source/EditableMesh/EditableGeometryCollectionAdapter.cpp b/Engine/Plugins/Runtime/EditableMesh/Source/EditableMesh/EditableGeometryCollectionAdapter.cpp
--- a/Engine/Plugins/Runtime/EditableMesh/Source/EditableMesh/EditableGeometryCollectionAdapter.cpp
+++ b/Engine/Plugins/Runtime/EditableMesh/Source/EditableMesh/EditableGeometryCollectionAdapter.cpp
@@ -6,7 +6,8 @@
 #include "Engine/StaticMesh.h"
 #include "GeometryCollection/GeometryCollectionComponent.h"
 #include "EditableMeshFactory.h"
-#include "MeshAttributes.h"
+#include "StaticMeshAttributes.h"
+#include "StaticMeshOperations.h"
 #include "EditorSupportDelegates.h"
 #include "GeometryCollection/GeometryCollectionAlgo.h"
 
@@ -241,14 +242,7 @@ void UEditableGeometryCollectionAdapter::InitEditableGeometryCollection( UEditab
 						}
 					}
 				}
-				MeshDescription->DetermineEdgeHardnessesFromVertexInstanceNormals();
-
-				// Determine UV seams
-				CA_SUPPRESS(6326);
-				if (NumUVs > 0)
-				{
-					MeshDescription->DetermineUVSeamsFromUVs(0);
-				}
+				FStaticMeshOperations::DetermineEdgeHardnessesFromVertexInstanceNormals(*MeshDescription);
 
 				// Cache polygon tangent bases
 				static TArray<FPolygonID> PolygonIDs;
diff --git a/Engine/Plugins/Runtime/EditableMesh/Source/EditableMesh/EditableStaticMeshAdapter.cpp b/Engine/Plugins/Runtime/EditableMesh/Source/EditableMesh/EditableStaticMeshAdapter.cpp
--- a/Engine/Plugins/Runtime/EditableMesh/Source/EditableMesh/EditableStaticMeshAdapter.cpp
+++ b/Engine/Plugins/Runtime/EditableMesh/Source/EditableMesh/EditableStaticMeshAdapter.cpp
@@ -9,7 +9,8 @@
 #include "PhysicsEngine/BodySetup.h"	// For collision generation
 #include "ProfilingDebugging/ScopedTimers.h"	// For FAutoScopedDurationTimer
 #include "EditableMeshFactory.h"
-#include "MeshAttributes.h"
+#include "StaticMeshAttributes.h"
+#include "StaticMeshOperations.h"
 
 
 UEditableStaticMeshAdapter::UEditableStaticMeshAdapter()
@@ -334,13 +335,7 @@ void UEditableStaticMeshAdapter::InitEditableStaticMesh( UEditableMesh* Editable
 				}
 
 				// Determine edge hardnesses
-				MeshDescription->DetermineEdgeHardnessesFromVertexInstanceNormals();
-
-				// Determine UV seams
-				if( NumUVs > 0 )
-				{
-					MeshDescription->DetermineUVSeamsFromUVs( 0 );
-				}
+				FStaticMeshOperations::DetermineEdgeHardnessesFromVertexInstanceNormals(*MeshDescription);
 
 				// Cache polygon tangent bases
 				static TArray<FPolygonID> PolygonIDs;
diff --git a/Engine/Plugins/Runtime/EditableMesh/Source/EditableMesh/GeometryCollectionEditableMeshFormat.cpp b/Engine/Plugins/Runtime/EditableMesh/Source/EditableMesh/GeometryCollectionEditableMeshFormat.cpp
--- a/Engine/Plugins/Runtime/EditableMesh/Source/EditableMesh/GeometryCollectionEditableMeshFormat.cpp
+++ b/Engine/Plugins/Runtime/EditableMesh/Source/EditableMesh/GeometryCollectionEditableMeshFormat.cpp
@@ -5,7 +5,7 @@
 #include "Engine/StaticMesh.h"
 #include "GeometryCollection/GeometryCollectionComponent.h"
 #include "EditableGeometryCollectionAdapter.h"
-#include "MeshAttributes.h"
+#include "StaticMeshAttributes.h"
 
 bool FGeometryCollectionEditableMeshFormat::HandlesComponentType(class UPrimitiveComponent& Component)
 {
@@ -123,7 +123,6 @@ void FGeometryCollectionEditableMeshFormat::RegisterMeshAttributes(FMeshDescript
 
 	// Add basic edge attributes
 	MeshDescription->EdgeAttributes().RegisterAttribute<bool>(MeshAttribute::Edge::IsHard, 1, false);
-	MeshDescription->EdgeAttributes().RegisterAttribute<bool>(MeshAttribute::Edge::IsUVSeam, 1, false);
 	MeshDescription->EdgeAttributes().RegisterAttribute<float>(MeshAttribute::Edge::CreaseSharpness, 1, 0.0f, EMeshAttributeFlags::Lerpable);
 
 	// Add basic polygon attributes
diff --git a/Engine/Plugins/Runtime/EditableMesh/Source/EditableMesh/StaticMeshEditableMeshFormat.cpp b/Engine/Plugins/Runtime/EditableMesh/Source/EditableMesh/StaticMeshEditableMeshFormat.cpp
--- a/Engine/Plugins/Runtime/EditableMesh/Source/EditableMesh/StaticMeshEditableMeshFormat.cpp
+++ b/Engine/Plugins/Runtime/EditableMesh/Source/EditableMesh/StaticMeshEditableMeshFormat.cpp
@@ -6,6 +6,8 @@
 #include "Components/StaticMeshComponent.h"
 #include "StaticMeshResources.h"
 #include "EditableStaticMeshAdapter.h"
+#include "StaticMeshAttributes.h"
+
 
 bool FStaticMeshEditableMeshFormat::HandlesComponentType(class UPrimitiveComponent& Component)
 {
@@ -56,10 +58,9 @@ UEditableMesh* FStaticMeshEditableMeshFormat::MakeEditableMesh( UPrimitiveCompon
 
 	UEditableMesh* EditableMesh = NewObject<UEditableMesh>();
 	FMeshDescription* MeshDescription = EditableMesh->GetMeshDescription();
-	UStaticMesh::RegisterMeshAttributes( *MeshDescription );
+	FStaticMeshAttributes( *MeshDescription ).Register();
 
 	// Register additional attributes required by EditableMesh
-	MeshDescription->EdgeAttributes().RegisterAttribute<bool>( MeshAttribute::Edge::IsUVSeam, 1, false );
 	MeshDescription->PolygonAttributes().RegisterAttribute<FVector>( MeshAttribute::Polygon::Normal, 1, FVector::ZeroVector, EMeshAttributeFlags::Transient );
 	MeshDescription->PolygonAttributes().RegisterAttribute<FVector>( MeshAttribute::Polygon::Tangent, 1, FVector::ZeroVector, EMeshAttributeFlags::Transient );
 	MeshDescription->PolygonAttributes().RegisterAttribute<FVector>( MeshAttribute::Polygon::Binormal, 1, FVector::ZeroVector, EMeshAttributeFlags::Transient );
diff --git a/Engine/Plugins/Runtime/ProceduralMeshComponent/Source/ProceduralMeshComponent/Private/ProceduralMeshConversion.cpp b/Engine/Plugins/Runtime/ProceduralMeshComponent/Source/ProceduralMeshComponent/Private/ProceduralMeshConversion.cpp
--- a/Engine/Plugins/Runtime/ProceduralMeshComponent/Source/ProceduralMeshComponent/Private/ProceduralMeshConversion.cpp
+++ b/Engine/Plugins/Runtime/ProceduralMeshComponent/Source/ProceduralMeshComponent/Private/ProceduralMeshConversion.cpp
@@ -3,96 +3,9 @@
 #include "ProceduralMeshConversion.h"
 #include "Engine/StaticMesh.h"
 #include "ProceduralMeshComponent.h"
-#include "MeshAttributes.h"
+#include "StaticMeshAttributes.h"
 #include "MeshDescription.h"
 
-struct FProceduralMeshDescriptionAttributeGetter
-{
-	FProceduralMeshDescriptionAttributeGetter(FMeshDescription* InMeshDescription)
-		: Mesh(InMeshDescription)
-	{}
-
-	FMeshDescription* Mesh;
-
-	TVertexAttributesRef<FVector>
-	GetPositions() const
-	{
-		return Mesh->VertexAttributes().GetAttributesRef<FVector>(MeshAttribute::Vertex::Position);
-	}
-	TVertexInstanceAttributesRef<FVector>
-	GetNormals() const
-	{
-		return Mesh->VertexInstanceAttributes().GetAttributesRef<FVector>(MeshAttribute::VertexInstance::Normal);
-	}
-	TVertexInstanceAttributesRef<FVector>
-	GetTangents() const
-	{
-		return Mesh->VertexInstanceAttributes().GetAttributesRef<FVector>(MeshAttribute::VertexInstance::Tangent);
-	}
-	TVertexInstanceAttributesRef<float>
-	GetBinormalSigns() const {
-		return Mesh->VertexInstanceAttributes().GetAttributesRef<float>(MeshAttribute::VertexInstance::BinormalSign);
-	}
-	TVertexInstanceAttributesRef<FVector4>
-	GetColors() const
-	{
-		return Mesh->VertexInstanceAttributes().GetAttributesRef<FVector4>(MeshAttribute::VertexInstance::Color);
-	}
-	TVertexInstanceAttributesRef<FVector2D>
-	GetUVs() const
-	{
-		return Mesh->VertexInstanceAttributes().GetAttributesRef<FVector2D>(MeshAttribute::VertexInstance::TextureCoordinate);
-	}
-	TPolygonGroupAttributesRef<FName>
-	GetPolygonGroupImportedMaterialSlotNames()
-	{
-		return Mesh->PolygonGroupAttributes().GetAttributesRef<FName>(MeshAttribute::PolygonGroup::ImportedMaterialSlotName);
-	}
-};
-
-struct FProceduralMeshDescriptionConstAttributeGetter
-{
-	FProceduralMeshDescriptionConstAttributeGetter(const FMeshDescription* InMeshDescription)
-		: Mesh(InMeshDescription)
-	{}
-
-	const FMeshDescription* Mesh;
-
-	TVertexAttributesConstRef<FVector>
-	GetPositions() const
-	{
-		return Mesh->VertexAttributes().GetAttributesRef<FVector>(MeshAttribute::Vertex::Position);
-	}
-	TVertexInstanceAttributesConstRef<FVector>
-	GetNormals() const
-	{
-		return Mesh->VertexInstanceAttributes().GetAttributesRef<FVector>(MeshAttribute::VertexInstance::Normal);
-	}
-	TVertexInstanceAttributesConstRef<FVector>
-	GetTangents() const
-	{
-		return Mesh->VertexInstanceAttributes().GetAttributesRef<FVector>(MeshAttribute::VertexInstance::Tangent);
-	}
-	TVertexInstanceAttributesConstRef<float>
-	GetBinormalSigns() const {
-		return Mesh->VertexInstanceAttributes().GetAttributesRef<float>(MeshAttribute::VertexInstance::BinormalSign);
-	}
-	TVertexInstanceAttributesConstRef<FVector4>
-	GetColors() const
-	{
-		return Mesh->VertexInstanceAttributes().GetAttributesRef<FVector4>(MeshAttribute::VertexInstance::Color);
-	}
-	TVertexInstanceAttributesConstRef<FVector2D>
-	GetUVs() const
-	{
-		return Mesh->VertexInstanceAttributes().GetAttributesRef<FVector2D>(MeshAttribute::VertexInstance::TextureCoordinate);
-	}
-	TPolygonGroupAttributesConstRef<FName>
-	GetPolygonGroupImportedMaterialSlotNames()
-	{
-		return Mesh->PolygonGroupAttributes().GetAttributesRef<FName>(MeshAttribute::PolygonGroup::ImportedMaterialSlotName);
-	}
-};
 
 TMap<UMaterialInterface*, FPolygonGroupID> BuildMaterialMap(UProceduralMeshComponent* ProcMeshComp, FMeshDescription& MeshDescription)
 {
@@ -100,8 +13,8 @@ TMap<UMaterialInterface*, FPolygonGroupID> BuildMaterialMap(UProceduralMeshCompo
 	const int32 NumSections = ProcMeshComp->GetNumSections();
 	UniqueMaterials.Reserve(NumSections);
 
-	FProceduralMeshDescriptionAttributeGetter AttributeGetter(&MeshDescription);
-	TPolygonGroupAttributesRef<FName> PolygonGroupNames = AttributeGetter.GetPolygonGroupImportedMaterialSlotNames();
+	FStaticMeshAttributes AttributeGetter(MeshDescription);
+	TPolygonGroupAttributesRef<FName> PolygonGroupNames = AttributeGetter.GetPolygonGroupMaterialSlotNames();
 	for (int32 SectionIdx = 0; SectionIdx < NumSections; SectionIdx++)
 	{
 		FProcMeshSection *ProcSection =
@@ -135,14 +48,14 @@ void MeshDescriptionToProcMesh( const FMeshDescription& MeshDescription, UProced
 {
 	ProcMeshComp->ClearAllMeshSections();
 
-	FProceduralMeshDescriptionConstAttributeGetter AttributeGetter(&MeshDescription);
-	TPolygonGroupAttributesConstRef<FName> PolygonGroupNames = AttributeGetter.GetPolygonGroupImportedMaterialSlotNames();
-	TVertexAttributesConstRef<FVector> VertexPositions = AttributeGetter.GetPositions();
-	TVertexInstanceAttributesConstRef<FVector> Tangents = AttributeGetter.GetTangents();
-	TVertexInstanceAttributesConstRef<float> BinormalSigns = AttributeGetter.GetBinormalSigns();
-	TVertexInstanceAttributesConstRef<FVector> Normals = AttributeGetter.GetNormals();
-	TVertexInstanceAttributesConstRef<FVector4> Colors = AttributeGetter.GetColors();
-	TVertexInstanceAttributesConstRef<FVector2D> UVs = AttributeGetter.GetUVs();
+	FStaticMeshConstAttributes AttributeGetter(MeshDescription);
+	TPolygonGroupAttributesConstRef<FName> PolygonGroupNames = AttributeGetter.GetPolygonGroupMaterialSlotNames();
+	TVertexAttributesConstRef<FVector> VertexPositions = AttributeGetter.GetVertexPositions();
+	TVertexInstanceAttributesConstRef<FVector> Tangents = AttributeGetter.GetVertexInstanceTangents();
+	TVertexInstanceAttributesConstRef<float> BinormalSigns = AttributeGetter.GetVertexInstanceBinormalSigns();
+	TVertexInstanceAttributesConstRef<FVector> Normals = AttributeGetter.GetVertexInstanceNormals();
+	TVertexInstanceAttributesConstRef<FVector4> Colors = AttributeGetter.GetVertexInstanceColors();
+	TVertexInstanceAttributesConstRef<FVector2D> UVs = AttributeGetter.GetVertexInstanceUVs();
 
 	const int32 NumSections = ProcMeshComp->GetNumSections();
 
@@ -188,16 +101,17 @@ void MeshDescriptionToProcMesh( const FMeshDescription& MeshDescription, UProced
 FMeshDescription BuildMeshDescription( UProceduralMeshComponent* ProcMeshComp )
 {
 	FMeshDescription MeshDescription;
-	UStaticMesh::RegisterMeshAttributes(MeshDescription);
 
-	FProceduralMeshDescriptionAttributeGetter AttributeGetter(&MeshDescription);
-	TPolygonGroupAttributesRef<FName> PolygonGroupNames = AttributeGetter.GetPolygonGroupImportedMaterialSlotNames();
-	TVertexAttributesRef<FVector> VertexPositions = AttributeGetter.GetPositions();
-	TVertexInstanceAttributesRef<FVector> Tangents = AttributeGetter.GetTangents();
-	TVertexInstanceAttributesRef<float> BinormalSigns = AttributeGetter.GetBinormalSigns();
-	TVertexInstanceAttributesRef<FVector> Normals = AttributeGetter.GetNormals();
-	TVertexInstanceAttributesRef<FVector4> Colors = AttributeGetter.GetColors();
-	TVertexInstanceAttributesRef<FVector2D> UVs = AttributeGetter.GetUVs();
+	FStaticMeshAttributes AttributeGetter(MeshDescription);
+	AttributeGetter.Register();
+
+	TPolygonGroupAttributesRef<FName> PolygonGroupNames = AttributeGetter.GetPolygonGroupMaterialSlotNames();
+	TVertexAttributesRef<FVector> VertexPositions = AttributeGetter.GetVertexPositions();
+	TVertexInstanceAttributesRef<FVector> Tangents = AttributeGetter.GetVertexInstanceTangents();
+	TVertexInstanceAttributesRef<float> BinormalSigns = AttributeGetter.GetVertexInstanceBinormalSigns();
+	TVertexInstanceAttributesRef<FVector> Normals = AttributeGetter.GetVertexInstanceNormals();
+	TVertexInstanceAttributesRef<FVector4> Colors = AttributeGetter.GetVertexInstanceColors();
+	TVertexInstanceAttributesRef<FVector2D> UVs = AttributeGetter.GetVertexInstanceUVs();
 
 	// Materials to apply to new mesh
 	const int32 NumSections = ProcMeshComp->GetNumSections();
diff --git a/Engine/Plugins/Runtime/ProceduralMeshComponent/Source/ProceduralMeshComponent/ProceduralMeshComponent.Build.cs b/Engine/Plugins/Runtime/ProceduralMeshComponent/Source/ProceduralMeshComponent/ProceduralMeshComponent.Build.cs
--- a/Engine/Plugins/Runtime/ProceduralMeshComponent/Source/ProceduralMeshComponent/ProceduralMeshComponent.Build.cs
+++ b/Engine/Plugins/Runtime/ProceduralMeshComponent/Source/ProceduralMeshComponent/ProceduralMeshComponent.Build.cs
@@ -17,7 +17,8 @@ namespace UnrealBuildTool.Rules
 					"InteractiveToolsFramework",
 					"MeshDescription",
 					"RenderCore",
-					"RHI"
+					"RHI",
+					"StaticMeshDescription"
 				}
 				);
 		}
diff --git a/Engine/Plugins/Runtime/ProceduralMeshComponent/Source/ProceduralMeshComponentEditor/Private/ProceduralMeshComponentDetails.cpp b/Engine/Plugins/Runtime/ProceduralMeshComponent/Source/ProceduralMeshComponentEditor/Private/ProceduralMeshComponentDetails.cpp
--- a/Engine/Plugins/Runtime/ProceduralMeshComponent/Source/ProceduralMeshComponentEditor/Private/ProceduralMeshComponentDetails.cpp
+++ b/Engine/Plugins/Runtime/ProceduralMeshComponent/Source/ProceduralMeshComponentEditor/Private/ProceduralMeshComponentDetails.cpp
@@ -17,10 +17,7 @@
 #include "DetailCategoryBuilder.h"
 #include "IDetailsView.h"
 #include "ProceduralMeshComponent.h"
-
-#include "MeshDescription.h"
-#include "MeshAttributes.h"
-#include "MeshAttributeArray.h"
+#include "StaticMeshAttributes.h"
 #include "MeshDescriptionOperations.h"
 
 #include "Dialogs/DlgPickAssetPath.h"
@@ -144,12 +141,7 @@ FReply FProceduralMeshComponentDetails::ClickedOnConvertToStaticMesh()
 				SrcModel.BuildSettings.bGenerateLightmapUVs = true;
 				SrcModel.BuildSettings.SrcLightmapIndex = 0;
 				SrcModel.BuildSettings.DstLightmapIndex = 1;
-				FMeshDescription* OriginalMeshDescription = StaticMesh->GetMeshDescription(0);
-				if (OriginalMeshDescription == nullptr)
-				{
-					OriginalMeshDescription = StaticMesh->CreateMeshDescription(0);
-				}
-				*OriginalMeshDescription = MeshDescription;
+				StaticMesh->CreateMeshDescription(0, MoveTemp(MeshDescription));
 				StaticMesh->CommitMeshDescription(0);
 
 				TSet<UMaterialInterface*> UniqueMaterials;
diff --git a/Engine/Plugins/Tests/EditorTests/Source/EditorTests/Private/UnrealEd/MeshDescriptionTests.cpp b/Engine/Plugins/Tests/EditorTests/Source/EditorTests/Private/UnrealEd/MeshDescriptionTests.cpp
--- a/Engine/Plugins/Tests/EditorTests/Source/EditorTests/Private/UnrealEd/MeshDescriptionTests.cpp
+++ b/Engine/Plugins/Tests/EditorTests/Source/EditorTests/Private/UnrealEd/MeshDescriptionTests.cpp
@@ -8,10 +8,10 @@
 #include "Modules/ModuleManager.h"
 #include "Engine/StaticMesh.h"
 #include "Materials/Material.h"
-#include "MeshDescription.h"
+#include "StaticMeshAttributes.h"
+#include "StaticMeshOperations.h"
 #include "MeshDescriptionOperations.h"
 #include "MeshBuilder.h"
-#include "MeshAttributes.h"
 #include "RawMesh.h"
 #include "MeshUtilities.h"
 
@@ -513,7 +513,7 @@ bool FMeshDescriptionTest::ConversionTest(FAutomationTestExecutionInfo& Executio
 				AssetMesh->GetSourceModel(LodIndex).LoadRawMesh(ResultRawMesh);
 				//Create a temporary Mesh Description
 				FMeshDescription MeshDescription;
-				UStaticMesh::RegisterMeshAttributes(MeshDescription);
+				FStaticMeshAttributes(MeshDescription).Register();
 				FMeshDescriptionOperations::ConvertFromRawMesh(ResultRawMesh, MeshDescription, MaterialMapInverse);
 				//Convert back the FRawmesh
 				FMeshDescriptionOperations::ConvertToRawMesh(MeshDescription, ResultRawMesh, MaterialMap);
@@ -575,9 +575,9 @@ bool FMeshDescriptionTest::NTBTest(FAutomationTestExecutionInfo& ExecutionInfo)
 		const TVertexInstanceAttributesRef<FVector2D> VertexInstanceUVs = MeshDescription.VertexInstanceAttributes().GetAttributesRef<FVector2D>(MeshAttribute::VertexInstance::TextureCoordinate);
 		int32 ExistingUVCount = VertexInstanceUVs.GetNumIndices();
 		//Build the normals and tangent and compare the result
-		MeshDescription.ComputePolygonTangentsAndNormals(SMALL_NUMBER);
-		EComputeNTBsOptions ComputeNTBsOptions = EComputeNTBsOptions::Normals | EComputeNTBsOptions::Tangents;
-		MeshDescription.ComputeTangentsAndNormals(ComputeNTBsOptions);
+		FStaticMeshOperations::ComputePolygonTangentsAndNormals(MeshDescription, SMALL_NUMBER);
+		EComputeNTBsFlags ComputeNTBsFlags = EComputeNTBsFlags::Normals | EComputeNTBsFlags::Tangents;
+		FStaticMeshOperations::ComputeTangentsAndNormals(MeshDescription, ComputeNTBsFlags);
 		//FMeshDescriptionOperations::CreatePolygonNTB(MeshDescription, SMALL_NUMBER);
 		//FMeshDescriptionOperations::CreateNormals(MeshDescription, FMeshDescriptionOperations::ETangentOptions::BlendOverlappingNormals, true);
 
diff --git a/Engine/Source/Developer/HierarchicalLODUtilities/HierarchicalLODUtilities.Build.cs b/Engine/Source/Developer/HierarchicalLODUtilities/HierarchicalLODUtilities.Build.cs
--- a/Engine/Source/Developer/HierarchicalLODUtilities/HierarchicalLODUtilities.Build.cs
+++ b/Engine/Source/Developer/HierarchicalLODUtilities/HierarchicalLODUtilities.Build.cs
@@ -19,6 +19,8 @@ public class HierarchicalLODUtilities : ModuleRules
             new string[]
 			{
 				"Engine",
+				"MeshDescription",
+				"StaticMeshDescription",
 				"UnrealEd",
                 "Projects",
 			}
diff --git a/Engine/Source/Developer/HierarchicalLODUtilities/Private/HierarchicalLODUtilities.cpp b/Engine/Source/Developer/HierarchicalLODUtilities/Private/HierarchicalLODUtilities.cpp
--- a/Engine/Source/Developer/HierarchicalLODUtilities/Private/HierarchicalLODUtilities.cpp
+++ b/Engine/Source/Developer/HierarchicalLODUtilities/Private/HierarchicalLODUtilities.cpp
@@ -3,6 +3,7 @@
 #include "HierarchicalLODUtilities.h"
 #include "GameFramework/Actor.h"
 #include "Components/StaticMeshComponent.h"
+#include "StaticMeshAttributes.h"
 #include "Modules/ModuleManager.h"
 #include "Misc/PackageName.h"
 #include "GameFramework/WorldSettings.h"
diff --git a/Engine/Source/Developer/MeshBoneReduction/Public/MeshBoneReduction.h b/Engine/Source/Developer/MeshBoneReduction/Public/MeshBoneReduction.h
--- a/Engine/Source/Developer/MeshBoneReduction/Public/MeshBoneReduction.h
+++ b/Engine/Source/Developer/MeshBoneReduction/Public/MeshBoneReduction.h
@@ -45,7 +45,7 @@ public:
 	 * @param	DesriedLOD	: The data to reduce comes from Skeleton 
 	 *
 	 */
-	virtual bool ReduceBoneCounts(USkeletalMesh * SkeletalMesh, int32 DesiredLOD, const TArray<FName>* BoneNamesToRemove) = 0;
+	virtual bool ReduceBoneCounts(USkeletalMesh * SkeletalMesh, int32 DesiredLOD, const TArray<FName>* BoneNamesToRemove, bool bCallPostEditChange = true) = 0;
 };
 
 /**
diff --git a/Engine/Source/Developer/MeshBuilder/MeshBuilder.Build.cs b/Engine/Source/Developer/MeshBuilder/MeshBuilder.Build.cs
--- a/Engine/Source/Developer/MeshBuilder/MeshBuilder.Build.cs
+++ b/Engine/Source/Developer/MeshBuilder/MeshBuilder.Build.cs
@@ -16,23 +16,20 @@ namespace UnrealBuildTool.Rules
                     "Engine",
                     "RenderCore",
                     "MeshDescription",
-                    "RenderCore",
+					"StaticMeshDescription",
                     "MeshDescriptionOperations",
                     "MeshReductionInterface",
                     "RawMesh",
+					"MeshUtilities",
                     "MeshUtilitiesCommon",
+                    "ClothingSystemRuntimeNv",
+                    "MeshBoneReduction",
+                    "SkeletalMeshUtilitiesCommon",
+					"MeshBuilderCommon",
                 }
 			);
 
-			AddEngineThirdPartyPrivateStaticDependencies(Target, "nvTriStrip");
-            AddEngineThirdPartyPrivateStaticDependencies(Target, "ForsythTriOptimizer");
-	        AddEngineThirdPartyPrivateStaticDependencies(Target, "nvTessLib");
             AddEngineThirdPartyPrivateStaticDependencies(Target, "QuadricMeshReduction");
-
-			if (Target.IsInPlatformGroup(UnrealPlatformGroup.Unix))
-			{
-				PublicSystemLibraries.Add("stdc++");	// can be fixed, see UE-70769
-			}
        }
 	}
 }
diff --git a/Engine/Source/Developer/MeshBuilder/Private/MeshBuilderModule.cpp b/Engine/Source/Developer/MeshBuilder/Private/MeshBuilderModule.cpp
--- a/Engine/Source/Developer/MeshBuilder/Private/MeshBuilderModule.cpp
+++ b/Engine/Source/Developer/MeshBuilder/Private/MeshBuilderModule.cpp
@@ -5,6 +5,8 @@
 
 #include "StaticMeshBuilder.h"
 #include "Engine/StaticMesh.h"
+#include "SkeletalMeshBuilder.h"
+#include "Engine/SkeletalMesh.h"
 
 class FMeshBuilderModule : public IMeshBuilderModule
 {
@@ -26,6 +28,8 @@ public:
 
 	virtual bool BuildMesh(FStaticMeshRenderData& OutRenderData, UObject* Mesh, const FStaticMeshLODGroup& LODGroup) override;
 
+	virtual bool BuildSkeletalMesh(USkeletalMesh* SkeletalMesh, const int32 LODIndex, const bool bRegenDepLODs) override;
+
 private:
 
 };
@@ -42,3 +46,9 @@ bool FMeshBuilderModule::BuildMesh(FStaticMeshRenderData& OutRenderData, class U
 	}
 	return false;
 }
+
+bool FMeshBuilderModule::BuildSkeletalMesh(USkeletalMesh* SkeletalMesh, const int32 LODIndex, const bool bRegenDepLODs)
+{
+	//Call the skeletal mesh builder
+	return FSkeletalMeshBuilder().Build(SkeletalMesh, LODIndex, bRegenDepLODs);
+}
\ No newline at end of file
diff --git a/Engine/Source/Developer/MeshBuilder/Private/MeshDescriptionHelper.cpp b/Engine/Source/Developer/MeshBuilder/Private/MeshDescriptionHelper.cpp
--- a/Engine/Source/Developer/MeshBuilder/Private/MeshDescriptionHelper.cpp
+++ b/Engine/Source/Developer/MeshBuilder/Private/MeshDescriptionHelper.cpp
@@ -8,8 +8,7 @@
 #include "UObject/Package.h"
 #include "Engine/EngineTypes.h"
 #include "Engine/StaticMesh.h"
-#include "MeshDescription.h"
-#include "MeshAttributes.h"
+#include "StaticMeshAttributes.h"
 #include "Serialization/MemoryWriter.h"
 #include "Serialization/MemoryReader.h"
 #include "RenderUtils.h"
diff --git a/Engine/Source/Developer/MeshBuilder/Private/SkeletalMeshBuilder.cpp b/Engine/Source/Developer/MeshBuilder/Private/SkeletalMeshBuilder.cpp
new file mode 100644
--- /dev/null
+++ b/Engine/Source/Developer/MeshBuilder/Private/SkeletalMeshBuilder.cpp
@@ -0,0 +1,197 @@
+// Copyright 1998-2019 Epic Games, Inc. All Rights Reserved.
+
+#include "SkeletalMeshBuilder.h"
+#include "Modules/ModuleManager.h"
+#include "MeshBoneReduction.h"
+#include "Engine/EngineTypes.h"
+#include "Engine/SkeletalMesh.h"
+#include "PhysicsEngine/BodySetup.h"
+#include "MeshDescription.h"
+#include "MeshDescriptionOperations.h"
+#include "MeshAttributes.h"
+#include "MeshDescriptionHelper.h"
+#include "MeshBuild.h"
+#include "Rendering/SkeletalMeshModel.h"
+#include "Rendering/SkeletalMeshLODModel.h"
+#include "GPUSkinVertexFactory.h"
+#include "ThirdPartyBuildOptimizationHelper.h"
+#include "Misc/ScopedSlowTask.h"
+#include "LODUtilities.h"
+#include "Assets/ClothingAsset.h"
+#include "MeshUtilities.h"
+#include "EditorFramework/AssetImportData.h"
+
+DEFINE_LOG_CATEGORY(LogSkeletalMeshBuilder);
+
+namespace SkeletalMeshBuilderOptimization
+{
+	void CacheOptimizeIndexBuffer(TArray<uint16>& Indices)
+	{
+		BuildOptimizationThirdParty::CacheOptimizeIndexBuffer(Indices);
+	}
+
+	void CacheOptimizeIndexBuffer(TArray<uint32>& Indices)
+	{
+		BuildOptimizationThirdParty::CacheOptimizeIndexBuffer(Indices);
+	}
+}
+
+struct InfluenceMap
+{
+	FORCEINLINE bool operator()(const float& A, const float& B) const
+	{
+		return B > A;
+	}
+};
+
+struct FSkeletalMeshVertInstanceIDAndZ
+{
+	FVertexInstanceID Index;
+	float Z;
+};
+
+FSkeletalMeshBuilder::FSkeletalMeshBuilder()
+{
+
+}
+
+
+bool FSkeletalMeshBuilder::Build(USkeletalMesh* SkeletalMesh, const int32 LODIndex, const bool bRegenDepLODs)
+{
+	check(SkeletalMesh->GetImportedModel());
+	check(SkeletalMesh->GetImportedModel()->LODModels.IsValidIndex(LODIndex));
+	check(SkeletalMesh->GetLODInfo(LODIndex) != nullptr);
+	
+	FSkeletalMeshLODInfo* LODInfo = SkeletalMesh->GetLODInfo(LODIndex);
+	//We want to backup in case the LODModel is regenerated, this data is use to validate in the UI if the ddc must be rebuild
+	const FString BackupBuildStringID = SkeletalMesh->GetImportedModel()->LODModels[LODIndex].BuildStringID;
+
+	const FReferenceSkeleton& RefSkeleton = SkeletalMesh->RefSkeleton;
+
+	FScopedSlowTask SlowTask(6.0f, NSLOCTEXT("SkeltalMeshBuilder", "BuildingSkeletalMeshLOD", "Building skeletal mesh LOD"));
+	SlowTask.MakeDialog();
+
+	//Prevent any PostEdit change during the build
+	FScopedSkeletalMeshPostEditChange ScopedPostEditChange(SkeletalMesh, false, false);
+	// Unbind any existing clothing assets before we reimport the geometry
+	TArray<ClothingAssetUtils::FClothingAssetMeshBinding> ClothingBindings;
+	FLODUtilities::UnbindClothingAndBackup(SkeletalMesh, ClothingBindings, LODIndex);
+
+	FSkeletalMeshImportData SkeletalMeshImportData;
+	float ProgressStepSize = SkeletalMesh->IsReductionActive(LODIndex) ? 0.5f : 1.0f;
+
+	//This scope define where we can use the LODModel, after a reduction the LODModel must be requery since it is a new instance
+	{
+		FSkeletalMeshLODModel& BuildLODModel = SkeletalMesh->GetImportedModel()->LODModels[LODIndex];
+		
+		BuildLODModel.RawSkeletalMeshBulkData.LoadRawMesh(SkeletalMeshImportData);
+
+		TArray<FVector> LODPoints;
+		TArray<SkeletalMeshImportData::FMeshWedge> LODWedges;
+		TArray<SkeletalMeshImportData::FMeshFace> LODFaces;
+		TArray<SkeletalMeshImportData::FVertInfluence> LODInfluences;
+		TArray<int32> LODPointToRawMap;
+		SkeletalMeshImportData.CopyLODImportData(LODPoints, LODWedges, LODFaces, LODInfluences, LODPointToRawMap);
+
+
+		//Build the skeletalmesh using mesh utilities module
+		IMeshUtilities::MeshBuildOptions Options;
+		Options.FillOptions(LODInfo->BuildSettings);
+		IMeshUtilities& MeshUtilities = FModuleManager::Get().LoadModuleChecked<IMeshUtilities>("MeshUtilities");
+
+		// Create skinning streams for NewModel.
+		SlowTask.EnterProgressFrame(1.0f);
+		MeshUtilities.BuildSkeletalMesh(
+			BuildLODModel,
+			RefSkeleton,
+			LODInfluences,
+			LODWedges,
+			LODFaces,
+			LODPoints,
+			LODPointToRawMap,
+			Options
+		);
+
+		//Re-Apply the user section changes, the UserSectionsData is map to original section and should match the builded LODModel
+		BuildLODModel.SyncronizeUserSectionsDataArray();
+
+		// Set texture coordinate count on the new model.
+		BuildLODModel.NumTexCoords = SkeletalMeshImportData.NumTexCoords;
+
+		//Re-apply the morph target
+		SlowTask.EnterProgressFrame(1.0f, NSLOCTEXT("SkeltalMeshBuilder", "RebuildMorphTarget", "Rebuilding morph targets..."));
+		if (SkeletalMeshImportData.MorphTargetNames.Num() > 0)
+		{
+			FLODUtilities::BuildMorphTargets(SkeletalMesh, SkeletalMeshImportData, LODIndex, !Options.bComputeNormals, !Options.bComputeTangents, Options.bUseMikkTSpace);
+		}
+
+		//Re-apply the alternate skinning it must be after the inline reduction
+		SlowTask.EnterProgressFrame(1.0f, NSLOCTEXT("SkeltalMeshBuilder", "RebuildAlternateSkinning", "Rebuilding alternate skinning..."));
+		const TArray<FSkinWeightProfileInfo>& SkinProfiles = SkeletalMesh->GetSkinWeightProfiles();
+		for (int32 SkinProfileIndex = 0; SkinProfileIndex < SkinProfiles.Num(); ++SkinProfileIndex)
+		{
+			const FSkinWeightProfileInfo& ProfileInfo = SkinProfiles[SkinProfileIndex];
+			FLODUtilities::UpdateAlternateSkinWeights(SkeletalMesh, ProfileInfo.Name, LODIndex, Options.OverlappingThresholds, !Options.bComputeNormals, !Options.bComputeTangents, Options.bUseMikkTSpace, Options.bComputeWeightedNormals);
+		}
+
+		
+		FSkeletalMeshUpdateContext UpdateContext;
+		UpdateContext.SkeletalMesh = SkeletalMesh;
+		//We are reduce ourself in this case we reduce ourself from the original data and return true.
+		if (SkeletalMesh->IsReductionActive(LODIndex))
+		{
+			SlowTask.EnterProgressFrame(ProgressStepSize, NSLOCTEXT("SkeltalMeshBuilder", "RegenerateLOD", "Regenerate LOD..."));
+			//Update the original reduction data since we just build a new LODModel.
+			if (LODInfo->ReductionSettings.BaseLOD == LODIndex && SkeletalMesh->GetImportedModel()->OriginalReductionSourceMeshData.IsValidIndex(LODIndex))
+			{
+				//Make the copy of the data only once until the ImportedModel change (re-imported)
+				SkeletalMesh->GetImportedModel()->OriginalReductionSourceMeshData[LODIndex]->EmptyBulkData();
+				TMap<FString, TArray<FMorphTargetDelta>> BaseLODMorphTargetData;
+				BaseLODMorphTargetData.Empty(SkeletalMesh->MorphTargets.Num());
+				for (UMorphTarget *MorphTarget : SkeletalMesh->MorphTargets)
+				{
+					if (!MorphTarget->HasDataForLOD(LODIndex))
+					{
+						continue;
+					}
+					TArray<FMorphTargetDelta>& MorphDeltasArray = BaseLODMorphTargetData.FindOrAdd(MorphTarget->GetFullName());
+					const FMorphTargetLODModel& BaseMorphModel = MorphTarget->MorphLODModels[LODIndex];
+					//Iterate each original morph target source index to fill the NewMorphTargetDeltas array with the TargetMatchData.
+					for (const FMorphTargetDelta& MorphDelta : BaseMorphModel.Vertices)
+					{
+						MorphDeltasArray.Add(MorphDelta);
+					}
+				}
+				//Copy the original SkeletalMesh LODModel
+				SkeletalMesh->GetImportedModel()->OriginalReductionSourceMeshData[LODIndex]->SaveReductionData(BuildLODModel, BaseLODMorphTargetData, SkeletalMesh);
+
+				if (LODIndex == 0)
+				{
+					SkeletalMesh->GetLODInfo(LODIndex)->SourceImportFilename = SkeletalMesh->AssetImportData->GetFirstFilename();
+				}
+			}
+			FLODUtilities::SimplifySkeletalMeshLOD(UpdateContext, LODIndex, false);
+		}
+	}
+
+	FSkeletalMeshLODModel& LODModelAfterReduction = SkeletalMesh->GetImportedModel()->LODModels[LODIndex];
+	
+	//Re-apply the clothing using the UserSectionsData, this will ensure we remap correctly the cloth if the reduction has change the number of sections
+	SlowTask.EnterProgressFrame(1.0f, NSLOCTEXT("SkeltalMeshBuilder", "RebuildClothing", "Rebuilding clothing..."));
+	FLODUtilities::RestoreClothingFromBackup(SkeletalMesh, ClothingBindings, LODIndex);
+
+	LODModelAfterReduction.SyncronizeUserSectionsDataArray();
+	LODModelAfterReduction.NumTexCoords = SkeletalMeshImportData.NumTexCoords;
+	LODModelAfterReduction.BuildStringID = BackupBuildStringID;
+
+	SlowTask.EnterProgressFrame(ProgressStepSize, NSLOCTEXT("SkeltalMeshBuilder", "RegenerateDependentLODs", "Regenerate Dependent LODs..."));
+	if (bRegenDepLODs)
+	{
+		//Regenerate dependent LODs
+		FLODUtilities::RegenerateDependentLODs(SkeletalMesh, LODIndex);
+	}
+
+	SlowTask.EnterProgressFrame();
+
+	return true;
+}
