Fix crash when reducing an empty mesh
#rnx

diff --git a/Engine/Plugins/Experimental/SkeletalReduction/Source/Private/SkeletalSimplifierMeshManager.cpp b/Engine/Plugins/Experimental/SkeletalReduction/Source/Private/SkeletalSimplifierMeshManager.cpp
--- a/Engine/Plugins/Experimental/SkeletalReduction/Source/Private/SkeletalSimplifierMeshManager.cpp
+++ b/Engine/Plugins/Experimental/SkeletalReduction/Source/Private/SkeletalSimplifierMeshManager.cpp
@@ -424,6 +424,11 @@ void SkeletalSimplifier::FSimplifierMeshManager::FlagBoundary(const ESimpElement
 { 
 
 	TArray< SimpVertType*, TInlineAllocator<64> > adjVerts;
+	if (NumSrcVerts == 0 || NumSrcTris == 0)
+	{
+		//Avoid trying to compute an empty mesh
+		return;
+	}
 
 	for (int i = 0; i < NumSrcVerts; i++)
 	{
