Fix crash when restoring disabled sections

diff --git a/Engine/Plugins/Experimental/SkeletalReduction/Source/Private/SkeletalMeshReductionPlugin.cpp b/Engine/Plugins/Experimental/SkeletalReduction/Source/Private/SkeletalMeshReductionPlugin.cpp
--- a/Engine/Plugins/Experimental/SkeletalReduction/Source/Private/SkeletalMeshReductionPlugin.cpp
+++ b/Engine/Plugins/Experimental/SkeletalReduction/Source/Private/SkeletalMeshReductionPlugin.cpp
@@ -1889,12 +1889,6 @@ void FQuadricSkeletalMeshReduction::ReduceSkeletalMesh(USkeletalMesh& SkeletalMe
 		SrcBackupSectionIndexToSectionData.Empty(LODModelToBackup.Sections.Num());
 		for (int32 SectionIndex = 0; SectionIndex < LODModelToBackup.Sections.Num(); ++SectionIndex)
 		{
-			//Skip disable and not generated section
-			int32 MaxGeneratedLODIndex = LODModelToBackup.Sections[SectionIndex].GenerateUpToLodIndex;
-			if (LODModelToBackup.Sections[SectionIndex].bDisabled || (MaxGeneratedLODIndex != -1 && MaxGeneratedLODIndex < LODIndex))
-			{
-				continue;
-			}
 			FSectionData& SectionData = SrcBackupSectionIndexToSectionData.FindOrAdd(SectionIndex);
 			SectionData.MaterialIndex = LODModelToBackup.Sections[SectionIndex].MaterialIndex;
 			SectionData.bCastShadow = LODModelToBackup.Sections[SectionIndex].bCastShadow;
@@ -2052,8 +2046,8 @@ void FQuadricSkeletalMeshReduction::ReduceSkeletalMesh(USkeletalMesh& SkeletalMe
 		//Restore the source sections data
 		{
 			FSkeletalMeshLODModel& ImportedModelLOD = SkeletalMesh.GetImportedModel()->LODModels[LODIndex];
-			TArray<bool> SectionMatched;
-			SectionMatched.AddZeroed(SrcBackupSectionIndexToSectionData.Num());
+			TMap<int32, bool> OriginalSectionMatched;
+			OriginalSectionMatched.Reserve(SrcBackupSectionIndexToSectionData.Num());
 			int32 CurrentParentSectionIndex = INDEX_NONE;
 			int32 OriginalSectionIndex = INDEX_NONE;
 			for (int32 SectionIndex = 0; SectionIndex < ImportedModelLOD.Sections.Num(); ++SectionIndex)
@@ -2062,7 +2056,8 @@ void FQuadricSkeletalMeshReduction::ReduceSkeletalMesh(USkeletalMesh& SkeletalMe
 				for (auto Kvp : SrcBackupSectionIndexToSectionData)
 				{
 					const int32 SourceSectionIndex = Kvp.Key;
-					if (SectionMatched[Kvp.Key])
+					bool& SectionMatched = OriginalSectionMatched.FindOrAdd(SourceSectionIndex);
+					if (SectionMatched)
 					{
 						continue;
 					}
@@ -2081,8 +2076,9 @@ void FQuadricSkeletalMeshReduction::ReduceSkeletalMesh(USkeletalMesh& SkeletalMe
 						Section.bDisabled = SectionData.bDisabled;
 						Section.GenerateUpToLodIndex = SectionData.GenerateUpToLodIndex;
 						Section.ChunkedParentSectionIndex = bIsChunkedSection ? CurrentParentSectionIndex : INDEX_NONE;
-						Section.OriginalDataSectionIndex = OriginalSectionIndex;
-						SectionMatched[Kvp.Key] = true; //a backup section can be restore only once
+						//If we reduce inline the source model, we want to use the real source original section
+						Section.OriginalDataSectionIndex = bReducingSourceModel ? SectionData.OriginalDataSectionIndex : OriginalSectionIndex;
+						SectionMatched = true; //a backup section can be restore only once
 						break;
 					}
 				}
