Make sure we restore the section data when we reduce a skeletalmesh

diff --git a/Engine/Plugins/Experimental/SkeletalReduction/Source/Private/SkeletalMeshReductionPlugin.cpp b/Engine/Plugins/Experimental/SkeletalReduction/Source/Private/SkeletalMeshReductionPlugin.cpp
--- a/Engine/Plugins/Experimental/SkeletalReduction/Source/Private/SkeletalMeshReductionPlugin.cpp
+++ b/Engine/Plugins/Experimental/SkeletalReduction/Source/Private/SkeletalMeshReductionPlugin.cpp
@@ -1772,7 +1772,34 @@ void FQuadricSkeletalMeshReduction::ReduceSkeletalMesh(USkeletalMesh& SkeletalMe
 			UE_LOG(LogSkeletalMeshReduction, Warning, TEXT("Building LOD %d - Invalid Base LOD entered. Using Base LOD 0 instead"), LODIndex);
 		}
 	}
-	
+
+	//Store the sections flags
+	struct FSectionData
+	{
+		uint16 MaterialIndex;
+		bool bCastShadow;
+		bool bRecomputeTangent;
+	};
+
+	TMap<int32, FSectionData> BackupSectionIndexToSectionData;
+	//Store the section data. Store the source LOD model in case we add a LOD, or the target LOD model in case the LOD already exist
+	FSkeletalMeshLODModel& BackupSectionLODModel = bLODModelAdded ? SkeletalMeshResource.LODModels[Settings.BaseLOD] : SkeletalMeshResource.LODModels[LODIndex];
+	BackupSectionIndexToSectionData.Reserve(BackupSectionLODModel.Sections.Num());
+
+	for (int32 SectionIndex = 0; SectionIndex < BackupSectionLODModel.Sections.Num(); ++SectionIndex)
+	{
+		//Skip disable and not generated section
+		int32 MaxGeneratedLODIndex = BackupSectionLODModel.Sections[SectionIndex].GenerateUpToLodIndex;
+		if (BackupSectionLODModel.Sections[SectionIndex].bDisabled || (MaxGeneratedLODIndex != -1 && MaxGeneratedLODIndex < LODIndex))
+		{
+			continue;
+		}
+		FSectionData& SectionData = BackupSectionIndexToSectionData.FindOrAdd(SectionIndex);
+		SectionData.MaterialIndex = BackupSectionLODModel.Sections[SectionIndex].MaterialIndex;
+		SectionData.bCastShadow = BackupSectionLODModel.Sections[SectionIndex].bCastShadow;
+		SectionData.bRecomputeTangent = BackupSectionLODModel.Sections[SectionIndex].bRecomputeTangent;
+	}
+
 	auto FillClothingData = [&SkeletalMeshResource, &LODIndex, bLODModelAdded](int32 &EnableSectionNumber, TArray<bool> &SectionStatus)
 	{
 		EnableSectionNumber = 0;
@@ -1963,6 +1990,31 @@ void FQuadricSkeletalMeshReduction::ReduceSkeletalMesh(USkeletalMesh& SkeletalMe
 		// Flag this LOD as having been simplified.
 		SkeletalMesh.GetLODInfo(LODIndex)->bHasBeenSimplified = true;
 		SkeletalMesh.bHasBeenSimplified = true;
+		//Restore section data
+		FSkeletalMeshLODModel& ImportedModelLOD = SkeletalMesh.GetImportedModel()->LODModels[LODIndex];
+		TArray<bool> SectionMatched;
+		SectionMatched.AddZeroed(ImportedModelLOD.Sections.Num());
+		for (auto Kvp : BackupSectionIndexToSectionData)
+		{
+			const int32 SourceSectionIndex = Kvp.Key;
+
+			const FSectionData& SectionData = Kvp.Value;
+			int32 MaxSectionIndex = FMath::Min(ImportedModelLOD.Sections.Num(), SourceSectionIndex+1);
+			for (int32 SectionIndex = 0; SectionIndex < MaxSectionIndex; ++SectionIndex)
+			{
+				if (SectionMatched[SectionIndex])
+				{
+					continue;
+				}
+				if (SectionData.MaterialIndex == ImportedModelLOD.Sections[SectionIndex].MaterialIndex)
+				{
+					ImportedModelLOD.Sections[SectionIndex].bCastShadow = SectionData.bCastShadow;
+					ImportedModelLOD.Sections[SectionIndex].bRecomputeTangent = SectionData.bRecomputeTangent;
+					SectionMatched[SectionIndex] = true;
+					break;
+				}
+			}
+		}
 	}
 	else
 	{
