Make the reduction backup section more robust
#rnx

diff --git a/Engine/Plugins/Experimental/SkeletalReduction/Source/Private/SkeletalMeshReductionPlugin.cpp b/Engine/Plugins/Experimental/SkeletalReduction/Source/Private/SkeletalMeshReductionPlugin.cpp
--- a/Engine/Plugins/Experimental/SkeletalReduction/Source/Private/SkeletalMeshReductionPlugin.cpp
+++ b/Engine/Plugins/Experimental/SkeletalReduction/Source/Private/SkeletalMeshReductionPlugin.cpp
@@ -1798,6 +1798,7 @@ void FQuadricSkeletalMeshReduction::ReduceSkeletalMesh(USkeletalMesh& SkeletalMe
 	struct FSectionData
 	{
 		uint16 MaterialIndex;
+		int32 MaterialMap;
 		bool bCastShadow;
 		bool bRecomputeTangent;
 		bool bDisabled;
@@ -1835,6 +1836,7 @@ void FQuadricSkeletalMeshReduction::ReduceSkeletalMesh(USkeletalMesh& SkeletalMe
 		SectionData.GenerateUpToLodIndex = BackupSectionLODModel.Sections[SectionIndex].GenerateUpToLodIndex;
 		SectionData.ChunkedParentSectionIndex = BackupSectionLODModel.Sections[SectionIndex].ChunkedParentSectionIndex;
 		SectionData.OriginalDataSectionIndex = BackupSectionLODModel.Sections[SectionIndex].OriginalDataSectionIndex;
+		SectionData.MaterialMap = bOldLodWasFromFile ? INDEX_NONE : LODInfo->LODMaterialMap.IsValidIndex(SectionIndex) ? LODInfo->LODMaterialMap[SectionIndex] : INDEX_NONE;
 	}
 
 	auto FillSectionMaterialSlot = [&SkeletalMeshResource, &LODIndex, bLODModelAdded](TArray<int32>& SectionMaterialSlot)
@@ -1981,6 +1983,8 @@ void FQuadricSkeletalMeshReduction::ReduceSkeletalMesh(USkeletalMesh& SkeletalMe
 
 	if (ReduceSkeletalLODModel(*SrcModel, *NewModel, SkeletalMesh.GetImportedBounds(), SkeletalMesh.RefSkeleton, Settings, ImportantBones, RelativeToRefPoseMatrices, LODIndex, bReducingSourceModel))
 	{
+		FSkeletalMeshLODInfo* ReducedLODInfoPtr = SkeletalMesh.GetLODInfo(LODIndex);
+		check(ReducedLODInfoPtr);
 		// Do any joint-welding / bone removal.
 
 		if (MeshBoneReductionInterface != NULL && MeshBoneReductionInterface->GetBoneReductionData(&SkeletalMesh, LODIndex, BonesToRemove))
@@ -1994,30 +1998,10 @@ void FQuadricSkeletalMeshReduction::ReduceSkeletalMesh(USkeletalMesh& SkeletalMe
 
 		if (bOldLodWasFromFile)
 		{
-			SkeletalMesh.GetLODInfo(LODIndex)->LODMaterialMap.Empty();
+			ReducedLODInfoPtr->LODMaterialMap.Empty();
 		}
-
-		// If base lod has a customized LODMaterialMap and this LOD doesn't (could have if changes are applied instead of freshly generated, copy over the data into new new LOD
-
-		if (SkeletalMesh.GetLODInfo(LODIndex)->LODMaterialMap.Num() == 0 && SkeletalMesh.GetLODInfo(BaseLOD)->LODMaterialMap.Num() != 0)
-		{
-			SkeletalMesh.GetLODInfo(LODIndex)->LODMaterialMap = SkeletalMesh.GetLODInfo(BaseLOD)->LODMaterialMap;
-		}
-		else
-		{
-			// Assuming the reducing step has set all material indices correctly, we double check if something went wrong
-			// make sure we don't have more materials
-			int32 TotalSectionCount = NewModel->Sections.Num();
-			if (SkeletalMesh.GetLODInfo(LODIndex)->LODMaterialMap.Num() > TotalSectionCount)
-			{
-				SkeletalMesh.GetLODInfo(LODIndex)->LODMaterialMap = SkeletalMesh.GetLODInfo(BaseLOD)->LODMaterialMap;
-				// Something went wrong during the reduce step during regenerate 					
-				check(SkeletalMesh.GetLODInfo(BaseLOD)->LODMaterialMap.Num() == TotalSectionCount || SkeletalMesh.GetLODInfo(BaseLOD)->LODMaterialMap.Num() == 0);
-			}
-		}
-
 		// Flag this LOD as having been simplified.
-		SkeletalMesh.GetLODInfo(LODIndex)->bHasBeenSimplified = true;
+		ReducedLODInfoPtr->bHasBeenSimplified = true;
 		SkeletalMesh.bHasBeenSimplified = true;
 		//Restore section data
 		FSkeletalMeshLODModel& ImportedModelLOD = SkeletalMesh.GetImportedModel()->LODModels[LODIndex];
@@ -2042,8 +2026,18 @@ void FQuadricSkeletalMeshReduction::ReduceSkeletalMesh(USkeletalMesh& SkeletalMe
 				{
 					continue;
 				}
-				if (SectionData.MaterialIndex == ImportedModelLOD.Sections[SectionIndex].MaterialIndex)
+				if (SectionData.OriginalDataSectionIndex == ImportedModelLOD.Sections[SectionIndex].OriginalDataSectionIndex)
 				{
+					//Restore the LOD material map for this section
+					if (SectionData.MaterialMap != INDEX_NONE)
+					{
+						while (ReducedLODInfoPtr->LODMaterialMap.Num() <= SectionIndex)
+						{
+							ReducedLODInfoPtr->LODMaterialMap.Add(INDEX_NONE);
+						}
+						ReducedLODInfoPtr->LODMaterialMap[SectionIndex] = SectionData.MaterialMap;
+					}
+					//Restore the LODModel Section
 					ImportedModelLOD.Sections[SectionIndex].bCastShadow = SectionData.bCastShadow;
 					ImportedModelLOD.Sections[SectionIndex].bRecomputeTangent = SectionData.bRecomputeTangent;
 					ImportedModelLOD.Sections[SectionIndex].bDisabled = SectionData.bDisabled;
diff --git a/Engine/Source/Editor/Persona/Private/PersonaMeshDetails.cpp b/Engine/Source/Editor/Persona/Private/PersonaMeshDetails.cpp
--- a/Engine/Source/Editor/Persona/Private/PersonaMeshDetails.cpp
+++ b/Engine/Source/Editor/Persona/Private/PersonaMeshDetails.cpp
@@ -92,6 +92,7 @@
 
 DEFINE_LOG_CATEGORY(LogSkeletalMeshPersonaMeshDetail);
 
+#define SUFFIXE_DEFAULT_MATERIAL TEXT(" (Default)")
 /*
 * Custom data key
 */
@@ -3672,7 +3673,8 @@ void FPersonaMeshDetails::OnGetSectionsForView(ISectionListBuilder& OutSections,
 		int32 NumSections = Model.Sections.Num();
 		for (int32 SectionIdx = 0; SectionIdx < NumSections; SectionIdx++)
 		{
-			int32 MaterialIndex = Model.Sections[SectionIdx].MaterialIndex;
+			int32 DefaultSectionMaterialIndex = Model.Sections[SectionIdx].MaterialIndex;
+			int32 MaterialIndex = DefaultSectionMaterialIndex;
 			if (MaterialMap.IsValidIndex(MaterialIndex) && SkelMesh->Materials.IsValidIndex(MaterialMap[MaterialIndex]))
 			{
 				MaterialIndex = MaterialMap[MaterialIndex];
@@ -3687,12 +3689,22 @@ void FPersonaMeshDetails::OnGetSectionsForView(ISectionListBuilder& OutSections,
 				for (const FSkeletalMaterial &SkeletalMaterial : SkelMesh->Materials)
 				{
 					if (MaterialIndex != CurrentIterMaterialIndex)
-						AvailableSectionName.Add(CurrentIterMaterialIndex, SkeletalMaterial.MaterialSlotName);
+					{
+						if (DefaultSectionMaterialIndex == CurrentIterMaterialIndex)
+						{
+							FString BuildDefaultName = SkeletalMaterial.MaterialSlotName.ToString() + SUFFIXE_DEFAULT_MATERIAL;
+							AvailableSectionName.Add(CurrentIterMaterialIndex, FName(*BuildDefaultName));
+						}
+						else
+						{
+							AvailableSectionName.Add(CurrentIterMaterialIndex, SkeletalMaterial.MaterialSlotName);
+						}
+					}
 					CurrentIterMaterialIndex++;
 				}
 				bool bClothSection = Model.Sections[SectionIdx].HasClothingData();
 				bool bIsChunkSection = Model.Sections[SectionIdx].ChunkedParentSectionIndex != INDEX_NONE;
-				OutSections.AddSection(LODIndex, SectionIdx, CurrentSectionMaterialSlotName, MaterialIndex, CurrentSectionOriginalImportedMaterialName, AvailableSectionName, SkelMesh->Materials[MaterialIndex].MaterialInterface, bClothSection, bIsChunkSection);
+				OutSections.AddSection(LODIndex, SectionIdx, CurrentSectionMaterialSlotName, MaterialIndex, CurrentSectionOriginalImportedMaterialName, AvailableSectionName, SkelMesh->Materials[MaterialIndex].MaterialInterface, bClothSection, bIsChunkSection, DefaultSectionMaterialIndex);
 			}
 		}
 	}
@@ -4885,11 +4897,15 @@ void FPersonaMeshDetails::OnSectionChanged(int32 LODIndex, int32 SectionIndex, i
 
 		check(TotalSectionCount > SectionIndex);
 
+		FString NewMaterialSlotNameString = NewMaterialSlotName.ToString();
+		NewMaterialSlotNameString.RemoveFromEnd(SUFFIXE_DEFAULT_MATERIAL, ESearchCase::CaseSensitive);
+		FName CleanNewMaterialSlotName(*NewMaterialSlotNameString);
+
 		int32 NewSkeletalMaterialIndex = INDEX_NONE;
 		FName NewImportedMaterialSlotName = NAME_None;
 		for (int SkeletalMaterialIndex = 0; SkeletalMaterialIndex < Mesh->Materials.Num(); ++SkeletalMaterialIndex)
 		{
-			if (NewMaterialSlotIndex == SkeletalMaterialIndex && Mesh->Materials[SkeletalMaterialIndex].MaterialSlotName == NewMaterialSlotName)
+			if (NewMaterialSlotIndex == SkeletalMaterialIndex && Mesh->Materials[SkeletalMaterialIndex].MaterialSlotName == CleanNewMaterialSlotName)
 			{
 				NewSkeletalMaterialIndex = SkeletalMaterialIndex;
 				NewImportedMaterialSlotName = Mesh->Materials[SkeletalMaterialIndex].ImportedMaterialSlotName;
diff --git a/Engine/Source/Editor/PropertyEditor/Private/PropertyCustomizationHelpers.cpp b/Engine/Source/Editor/PropertyEditor/Private/PropertyCustomizationHelpers.cpp
--- a/Engine/Source/Editor/PropertyEditor/Private/PropertyCustomizationHelpers.cpp
+++ b/Engine/Source/Editor/PropertyEditor/Private/PropertyCustomizationHelpers.cpp
@@ -815,9 +815,9 @@ public:
 	* @param Section		The Section being used
 	* @param bCanBeReplced	Whether or not the Section can be replaced by a user
 	*/
-	virtual void AddSection(int32 LodIndex, int32 SectionIndex, FName InMaterialSlotName, int32 InMaterialSlotIndex, FName InOriginalMaterialSlotName, const TMap<int32, FName> &InAvailableMaterialSlotName, const UMaterialInterface* Material, bool IsSectionUsingCloth, bool bIsChunkSection) override
+	virtual void AddSection(int32 LodIndex, int32 SectionIndex, FName InMaterialSlotName, int32 InMaterialSlotIndex, FName InOriginalMaterialSlotName, const TMap<int32, FName> &InAvailableMaterialSlotName, const UMaterialInterface* Material, bool IsSectionUsingCloth, bool bIsChunkSection, int32 DefaultMaterialIndex) override
 	{
-		FSectionListItem SectionItem(LodIndex, SectionIndex, InMaterialSlotName, InMaterialSlotIndex, InOriginalMaterialSlotName, InAvailableMaterialSlotName, Material, IsSectionUsingCloth, ThumbnailSize, bIsChunkSection);
+		FSectionListItem SectionItem(LodIndex, SectionIndex, InMaterialSlotName, InMaterialSlotIndex, InOriginalMaterialSlotName, InAvailableMaterialSlotName, Material, IsSectionUsingCloth, ThumbnailSize, bIsChunkSection, DefaultMaterialIndex);
 		if (!Sections.Contains(SectionItem))
 		{
 			Sections.Add(SectionItem);
@@ -934,7 +934,10 @@ public:
 
 	TSharedRef<SWidget> CreateValueContent(const TSharedPtr<FAssetThumbnailPool>& ThumbnailPool)
 	{
-		FText MaterialSlotNameTooltipText = SectionItem.IsSectionUsingCloth ? FText(LOCTEXT("SectionIndex_MaterialSlotNameTooltip", "Cannot change the material slot when the mesh section use the cloth system.")) : FText::GetEmpty();
+		FFormatNamedArguments Arguments;
+		Arguments.Add(TEXT("DefaultMaterialIndex"), SectionItem.DefaultMaterialIndex);
+		FText BaseMaterialSlotTooltip = SectionItem.DefaultMaterialIndex != SectionItem.MaterialSlotIndex ? FText::Format(LOCTEXT("SectionIndex_BaseMaterialSlotNameTooltip", "This section material slot was change from the default value [{DefaultMaterialIndex}]."), Arguments) : FText::GetEmpty();
+		FText MaterialSlotNameTooltipText = SectionItem.IsSectionUsingCloth ? FText(LOCTEXT("SectionIndex_MaterialSlotNameTooltip", "Cannot change the material slot when the mesh section use the cloth system.")) : BaseMaterialSlotTooltip;
 		return
 			SNew(SHorizontalBox)
 			+ SHorizontalBox::Slot()
@@ -1093,7 +1096,12 @@ private:
 	{
 		FString MaterialSlotDisplayName;
 		SectionItem.MaterialSlotName.ToString(MaterialSlotDisplayName);
-		MaterialSlotDisplayName = TEXT("[") + FString::FromInt(SectionItem.MaterialSlotIndex) + TEXT("] ") + MaterialSlotDisplayName;
+		FString MaterialSlotRemapString = TEXT("");
+		if (SectionItem.DefaultMaterialIndex != SectionItem.MaterialSlotIndex)
+		{
+			MaterialSlotRemapString = TEXT(" (Modified)");
+		}
+		MaterialSlotDisplayName = TEXT("[") + FString::FromInt(SectionItem.MaterialSlotIndex) + TEXT("] ") + MaterialSlotDisplayName + MaterialSlotRemapString;
 		return FText::FromString(MaterialSlotDisplayName);
 	}
 
@@ -1227,7 +1235,7 @@ void FSectionList::GenerateChildContent(IDetailChildrenBuilder& ChildrenBuilder)
 			if (bDisplayAllSectionsInSlot)
 			{
 				FDetailWidgetRow& ChildRow = ChildrenBuilder.AddCustomRow(Section.Material.IsValid() ? FText::FromString(Section.Material->GetName()) : FText::GetEmpty());
-				AddSectionItem(ChildRow, CurrentLODIndex, FSectionListItem(CurrentLODIndex, Section.SectionIndex, Section.MaterialSlotName, Section.MaterialSlotIndex, Section.OriginalMaterialSlotName, Section.AvailableMaterialSlotName, Section.Material.Get(), Section.IsSectionUsingCloth, ThumbnailSize, Section.bIsChunkSection), !bDisplayAllSectionsInSlot);
+				AddSectionItem(ChildRow, CurrentLODIndex, FSectionListItem(CurrentLODIndex, Section.SectionIndex, Section.MaterialSlotName, Section.MaterialSlotIndex, Section.OriginalMaterialSlotName, Section.AvailableMaterialSlotName, Section.Material.Get(), Section.IsSectionUsingCloth, ThumbnailSize, Section.bIsChunkSection, Section.DefaultMaterialIndex), !bDisplayAllSectionsInSlot);
 			}
 		}
 	}
diff --git a/Engine/Source/Editor/PropertyEditor/Public/PropertyCustomizationHelpers.h b/Engine/Source/Editor/PropertyEditor/Public/PropertyCustomizationHelpers.h
--- a/Engine/Source/Editor/PropertyEditor/Public/PropertyCustomizationHelpers.h
+++ b/Engine/Source/Editor/PropertyEditor/Public/PropertyCustomizationHelpers.h
@@ -576,7 +576,7 @@ public:
 	* @param Section		The Section being used
 	* @param bCanBeReplced	Whether or not the Section can be replaced by a user
 	*/
-	virtual void AddSection(int32 LodIndex, int32 SectionIndex, FName InMaterialSlotName, int32 InMaterialSlotIndex, FName InOriginalMaterialSlotName, const TMap<int32, FName> &InAvailableMaterialSlotName, const UMaterialInterface* Material, bool IsSectionUsingCloth, bool bIsChunkSection) = 0;
+	virtual void AddSection(int32 LodIndex, int32 SectionIndex, FName InMaterialSlotName, int32 InMaterialSlotIndex, FName InOriginalMaterialSlotName, const TMap<int32, FName> &InAvailableMaterialSlotName, const UMaterialInterface* Material, bool IsSectionUsingCloth, bool bIsChunkSection, int32 DefaultMaterialIndex) = 0;
 };
 
 
@@ -608,9 +608,10 @@ struct FSectionListItem
 	TMap<int32, FName> AvailableMaterialSlotName;
 
 	bool bIsChunkSection;
+	int32 DefaultMaterialIndex;
 
 
-	FSectionListItem(int32 InLodIndex, int32 InSectionIndex, FName InMaterialSlotName, int32 InMaterialSlotIndex, FName InOriginalMaterialSlotName, const TMap<int32, FName> &InAvailableMaterialSlotName, const UMaterialInterface* InMaterial, bool InIsSectionUsingCloth, int32 InThumbnailSize, bool InIsChunkSection)
+	FSectionListItem(int32 InLodIndex, int32 InSectionIndex, FName InMaterialSlotName, int32 InMaterialSlotIndex, FName InOriginalMaterialSlotName, const TMap<int32, FName> &InAvailableMaterialSlotName, const UMaterialInterface* InMaterial, bool InIsSectionUsingCloth, int32 InThumbnailSize, bool InIsChunkSection, int32 InDefaultMaterialIndex)
 		: LodIndex(InLodIndex)
 		, SectionIndex(InSectionIndex)
 		, IsSectionUsingCloth(InIsSectionUsingCloth)
@@ -621,11 +622,20 @@ struct FSectionListItem
 		, OriginalMaterialSlotName(InOriginalMaterialSlotName)
 		, AvailableMaterialSlotName(InAvailableMaterialSlotName)
 		, bIsChunkSection(InIsChunkSection)
+		, DefaultMaterialIndex(InDefaultMaterialIndex)
 	{}
 
 	bool operator==(const FSectionListItem& Other) const
 	{
-		bool IsSectionItemEqual = LodIndex == Other.LodIndex && SectionIndex == Other.SectionIndex && MaterialSlotIndex == Other.MaterialSlotIndex && MaterialSlotName == Other.MaterialSlotName && Material == Other.Material && AvailableMaterialSlotName.Num() == Other.AvailableMaterialSlotName.Num() && IsSectionUsingCloth == Other.IsSectionUsingCloth;
+		bool IsSectionItemEqual = LodIndex == Other.LodIndex
+			&& SectionIndex == Other.SectionIndex
+			&& MaterialSlotIndex == Other.MaterialSlotIndex
+			&& MaterialSlotName == Other.MaterialSlotName
+			&& Material == Other.Material
+			&& AvailableMaterialSlotName.Num() == Other.AvailableMaterialSlotName.Num()
+			&& IsSectionUsingCloth == Other.IsSectionUsingCloth
+			&& bIsChunkSection == Other.bIsChunkSection
+			&& DefaultMaterialIndex == Other.DefaultMaterialIndex;
 		if (IsSectionItemEqual)
 		{
 			for (auto Kvp : AvailableMaterialSlotName)
diff --git a/Engine/Source/Editor/StaticMeshEditor/Private/StaticMeshEditorTools.cpp b/Engine/Source/Editor/StaticMeshEditor/Private/StaticMeshEditorTools.cpp
--- a/Engine/Source/Editor/StaticMeshEditor/Private/StaticMeshEditorTools.cpp
+++ b/Engine/Source/Editor/StaticMeshEditor/Private/StaticMeshEditorTools.cpp
@@ -1971,7 +1971,8 @@ void FMeshSectionSettingsLayout::OnGetSectionsForView(ISectionListBuilder& OutSe
 				{
 					SectionMaterial = UMaterial::GetDefaultMaterial(MD_Surface);
 				}
-				OutSections.AddSection(LODIndex, SectionIndex, CurrentSectionMaterialSlotName, MaterialIndex, CurrentSectionOriginalImportedMaterialName, AvailableSectionName, StaticMesh.StaticMaterials[MaterialIndex].MaterialInterface, false, false);
+				//TODO: Need to know if a section material slot assignment was change from the default value (implemented in skeletalmesh editor)
+				OutSections.AddSection(LODIndex, SectionIndex, CurrentSectionMaterialSlotName, MaterialIndex, CurrentSectionOriginalImportedMaterialName, AvailableSectionName, StaticMesh.StaticMaterials[MaterialIndex].MaterialInterface, false, false, MaterialIndex);
 			}
 		}
 	}
