Skeletal Mesh Reduction.  Add ability to lock vertices at vertex color boundaries.  If an edge connects two different vertex colors, the vertices are locked.
This appears as a checkbox, and is off by default.

diff --git a/Engine/Plugins/Experimental/SkeletalReduction/Source/Private/SkeletalMeshReductionPlugin.cpp b/Engine/Plugins/Experimental/SkeletalReduction/Source/Private/SkeletalMeshReductionPlugin.cpp
--- a/Engine/Plugins/Experimental/SkeletalReduction/Source/Private/SkeletalMeshReductionPlugin.cpp
+++ b/Engine/Plugins/Experimental/SkeletalReduction/Source/Private/SkeletalMeshReductionPlugin.cpp
@@ -944,6 +944,7 @@ float FQuadricSkeletalMeshReduction::SimplifyMesh( const FSkeletalMeshOptimizati
 	const bool bLockEdges             = Settings.bLockEdges;
 	const bool bPreserveVolume        = (VolumeImportance > 1.e-4);
 	const bool bEnforceBoneBoundaries = Settings.bEnforceBoneBoundaries;
+	const bool bLockColorBoundaries   = Settings.bLockColorBounaries;
 
 	// Terminator tells the simplifier when to stop
 	SkeletalSimplifier::FSimplifierTerminator Terminator(MinTriNumToRetain, SrcTriNum, MinVerNumToRetain, SrcVertNum, MaxCollapseCost, MaxDist);
@@ -1047,7 +1048,7 @@ float FQuadricSkeletalMeshReduction::SimplifyMesh( const FSkeletalMeshOptimizati
 	const float CoAlignmentLimit = FMath::Cos(45.f * PI / 180.); // 45 degrees limit
 
 	// Create the simplifier
-
+	
 	SkeletalSimplifier::FMeshSimplifier  Simplifier(Mesh.VertexBuffer, (uint32)Mesh.NumVertices(),
 		                                            Mesh.IndexBuffer, (uint32)Mesh.NumIndices(), 
 		                                            CoAlignmentLimit, VolumeImportance, bPreserveVolume,  bEnforceBoneBoundaries);
@@ -1079,6 +1080,12 @@ float FQuadricSkeletalMeshReduction::SimplifyMesh( const FSkeletalMeshOptimizati
 			Simplifier.SetBoundaryLocked();
 		}
 
+		if (bLockColorBoundaries)
+		{
+			// Lock the verts in edges that connect differnt colors.  Also locks verts that have multiple colors.
+			Simplifier.SetColorEdgeLocked();
+		}
+
 	}
 
 	// Do the actual simplification
diff --git a/Engine/Plugins/Experimental/SkeletalReduction/Source/Private/SkeletalSimplifier.cpp b/Engine/Plugins/Experimental/SkeletalReduction/Source/Private/SkeletalSimplifier.cpp
--- a/Engine/Plugins/Experimental/SkeletalReduction/Source/Private/SkeletalSimplifier.cpp
+++ b/Engine/Plugins/Experimental/SkeletalReduction/Source/Private/SkeletalSimplifier.cpp
@@ -75,6 +75,28 @@ void SkeletalSimplifier::FMeshSimplifier::SetBoundaryLocked()
 
 }
 
+void SkeletalSimplifier::FMeshSimplifier::SetColorEdgeLocked(float ColorDistThreshold)
+{
+
+	auto IsDifferntColor = [ColorDistThreshold](const SimpVertType* ASimpVert, const SimpVertType* BSimpVert)->bool
+	{
+		bool Result = true;
+		if (ASimpVert!=nullptr && BSimpVert != nullptr)
+		{ 
+			const FLinearColor& AColor = ASimpVert->vert.BasicAttributes.Color;
+			const FLinearColor& BColor = BSimpVert->vert.BasicAttributes.Color;
+
+		
+			Result = (FLinearColor::Dist(AColor, BColor) > ColorDistThreshold);
+		}
+		return Result;
+	};
+
+	MeshManager.FlagEdges(IsDifferntColor, ESimpElementFlags::SIMP_LOCKED);
+
+
+}
+
 void SkeletalSimplifier::FMeshSimplifier::SetBoxCornersLocked()
 {
 
diff --git a/Engine/Plugins/Experimental/SkeletalReduction/Source/Private/SkeletalSimplifier.h b/Engine/Plugins/Experimental/SkeletalReduction/Source/Private/SkeletalSimplifier.h
--- a/Engine/Plugins/Experimental/SkeletalReduction/Source/Private/SkeletalSimplifier.h
+++ b/Engine/Plugins/Experimental/SkeletalReduction/Source/Private/SkeletalSimplifier.h
@@ -76,7 +76,7 @@ namespace SkeletalSimplifier
 
 	};
 
-    /**
+	/**
 	* Class to allow a single weight value to be applied to all the sparse attributes.
 	*/
 	class UniformWeights
@@ -96,7 +96,7 @@ namespace SkeletalSimplifier
 			Weight = other.Weight;
 			return *this;
 		}
-		
+
 	private:
 		double Weight;
 
@@ -112,8 +112,8 @@ namespace SkeletalSimplifier
 		typedef Quadrics::TFaceQuadric<MeshVertType, UniformWeights>         WedgeQuadricType;
 		typedef TQuadricCache<WedgeQuadricType>                              QuadricCacheType;
 		typedef typename MeshVertType::BasicAttrContainerType                BasicAttrContainerType;
-		
-		
+
+
 
 		typedef typename MeshVertType::AttrContainerType                     AttrContainerType;
 
@@ -132,13 +132,13 @@ namespace SkeletalSimplifier
 
 
 	public:
-	
+
 		FMeshSimplifier(const MeshVertType* InSrcVerts, const uint32 InNumSrcVerts,
-			            const uint32* InSrcIndexes,     const uint32 InNumSrcIndexes,
-		            	const float CoAlignmentLimit,
-			            const float VolumeImportance,
-			            const bool bVolumePreservation,
-			            const bool bEnforceBoundaries);
+			const uint32* InSrcIndexes, const uint32 InNumSrcIndexes,
+			const float CoAlignmentLimit,
+			const float VolumeImportance,
+			const bool bVolumePreservation,
+			const bool bEnforceBoundaries);
 
 		~FMeshSimplifier();
 
@@ -152,19 +152,24 @@ namespace SkeletalSimplifier
 		* Set Quadric weights for the standard attributes ( the attributes that are stored in TBasicVertexAttrs<NumTexCoords> )
 		* The order of the weights corresponds to the order of the attributes :
 		*  Normal (3 weights), Tangent (3 weights ), BiTangent (3 weights), Color (3 weights), TextureCoords( 2 * NumTexCoords weights).
-		*/ 
+		*/
 		void				SetAttributeWeights(const DenseVecDType& Weights);
 
 		/**
-		* Set Quadric weights for sparse attributes 
+		* Set Quadric weights for sparse attributes
 		*/
 		void                SetSparseAttributeWeights(const SparseWeightContainerType& SparseWeights);
-		
+
 		/**
 		* Lock mesh boundary edges to prevent simplification by using the vertex tag ESimpElementFlags::SIMP_LOCKED
 		*/
 		void				SetBoundaryLocked();
-	
+
+		/**
+		* Lock edges that connect different colors to prevent simplification by using the vertex tag ESimpElementFlags::SIMP_LOCKED
+		*/
+		void                SetColorEdgeLocked(float Threshold = 1.e-3);
+
 		/**
 		* Lock vertices at corners of very simple 8-vert boxes.
 		*/
@@ -179,7 +184,7 @@ namespace SkeletalSimplifier
 		*                inline bool operator()(const int32 TriNum, const float SqrError);
 		*             };
 		*/
-		                    template <typename TerminationCriterionType>
+		template <typename TerminationCriterionType>
 		float				SimplifyMesh(TerminationCriterionType TerminationCriterion);
 
 		/**
@@ -193,7 +198,7 @@ namespace SkeletalSimplifier
 		int					GetNumTris() const { return MeshManager.ReducedNumTris; }
 
 		/**
-		* Export a copy of the simplified mesh. 
+		* Export a copy of the simplified mesh.
 		*
 		* @param Verts                 - pointer to an array to populate.  Should be least GetNumVerts() in size
 		* @param Indexes               - pointer to index buffer to populate.  Should be at least 3 * GetNumTris() in size
@@ -213,7 +218,7 @@ namespace SkeletalSimplifier
 		float               CalculateNormalShift(const SimpTriType& tri, const SimpVertType* vertToUpdate, const FVector& newPos) const;
 
 		void				ComputeEdgeCollapseVerts(SimpEdgeType* edge, EdgeUpdateTupleArray& newVerts);
-		void				ComputeEdgeCollapseVertsAndFixBones(SimpEdgeType* edge, EdgeUpdateTupleArray& newVerts); 
+		void				ComputeEdgeCollapseVertsAndFixBones(SimpEdgeType* edge, EdgeUpdateTupleArray& newVerts);
 
 
 		double				ComputeEdgeCollapseVertsAndCost(SimpEdgeType* edge, EdgeUpdateTupleArray& newVerts);
@@ -240,7 +245,13 @@ namespace SkeletalSimplifier
 		int32               CountDegenerates() const; // Included for testing.
 
 
-	
+		// Snap the interpolated vertex color in the third tuple element to the closest of the other two
+		void                SnapToClosestSourceVertexColor(EdgeUpdateTuple& UpdateTuple);
+
+		// Returns the closer (AColor or BColor) from the SampleColor
+		const FLinearColor& SnapToClosestColor(const FLinearColor& AColor, const FLinearColor& BColor, const FLinearColor& SampleColor) const;
+
+
 	protected:
 
 		// --- Weights for quadric simplification
@@ -252,9 +263,9 @@ namespace SkeletalSimplifier
 		SparseWeightContainerType AdditionalAttrWeights;
 
 
-	    // --- Magic numbers that penalize undesirable simplifications.
+		// --- Magic numbers that penalize undesirable simplifications.
 
-	    // prevent high valence verts.
+		// prevent high valence verts.
 		uint32               DegreeLimit = 24;
 		double               DegreePenalty = 100.0f;
 
@@ -262,7 +273,7 @@ namespace SkeletalSimplifier
 		double               invalidPenalty = 1.e6;
 
 		// critical angle by which the normal of a triangle is allowed to changed during a collapse.
-		double               coAlignmentLimit = .0871557f;  
+		double               coAlignmentLimit = .0871557f;
 
 		double               VolumeImportance;
 
@@ -273,10 +284,10 @@ namespace SkeletalSimplifier
 
 
 
-        // --- End Magic numbers
+		// --- End Magic numbers
 
 
-	    // Heap that maps edges to the cost of collapse
+		// Heap that maps edges to the cost of collapse
 		FBinaryHeap<double>		CollapseCostHeap;
 
 		// Manage the quadrics
@@ -286,7 +297,7 @@ namespace SkeletalSimplifier
 		// The mesh manager - holds a mesh that supports the topology queries we need and the collapse methods.
 		FSimplifierMeshManager   MeshManager;
 
-	
+
 	private:
 		FMeshSimplifier();
 
@@ -335,13 +346,13 @@ namespace SkeletalSimplifier
 		{
 			return Quadrics::FEdgeQuadric(Pos0, Pos1, Normal, Weight);
 		};
-	
+
 		Quadrics::FEdgeQuadric Quadric = QuadricCache.GetEdgeQuadric(v, EdgeQuadricFactory);
 
 		return Quadric;
 	}
 
-	
+
 
 	FORCEINLINE void FMeshSimplifier::ComputeEdgeCollapseVertsAndQuadrics(SimpEdgeType* edge,
 		EdgeUpdateTupleArray& EdgeAndNewVertArray,
@@ -369,6 +380,32 @@ namespace SkeletalSimplifier
 
 	}
 
+	FORCEINLINE const FLinearColor& FMeshSimplifier::SnapToClosestColor(const FLinearColor& AColor, const FLinearColor& BColor, const FLinearColor& SampleColor) const
+	{
+		// compute the squared distance in linear space - ignoring the alpha channel.
+		float ADistSqrd = (AColor.R - SampleColor.R)*(AColor.R - SampleColor.R) + (AColor.G - SampleColor.G)*(AColor.G - SampleColor.G) + (AColor.B - SampleColor.B)*(AColor.B - SampleColor.B);
+		float BDistSqrt = (BColor.R - SampleColor.R)*(BColor.R - SampleColor.R) + (BColor.G - SampleColor.G)*(BColor.G - SampleColor.G) + (BColor.B - SampleColor.B)*(BColor.B - SampleColor.B);
+
+		// select the closest
+		return (ADistSqrd > BDistSqrt) ? BColor : AColor;
+
+	}
+
+	FORCEINLINE void FMeshSimplifier::SnapToClosestSourceVertexColor(EdgeUpdateTuple& UpdateTuple)
+	{
+		MeshVertType&  VertAttributes = UpdateTuple.Get<2>();
+		FLinearColor& InterpolatedColor = VertAttributes.BasicAttributes.Color;
+
+		const SimpVertType* AVtxPtr = UpdateTuple.Get<0>();
+		const SimpVertType* BVtxPtr = UpdateTuple.Get<1>();
+
+		if (AVtxPtr != nullptr && BVtxPtr != nullptr)
+		{
+			// update to the closest source color.  
+			InterpolatedColor = SnapToClosestColor(AVtxPtr->vert.BasicAttributes.Color, BVtxPtr->vert.BasicAttributes.Color, InterpolatedColor);
+		}
+		
+	}
 
 	FORCEINLINE void FMeshSimplifier::UpdateEdgeCollapseCost(EdgePtrArray& DirtyEdges)
 	{
@@ -518,28 +555,29 @@ namespace SkeletalSimplifier
 
 				}
 
-
+				
 				// Update both verts in the mesh to have the new location and attribute values.
 
 				for (int32 i = 0, Imax = VertexUpdateArray.Num(); i < Imax; ++i)
 				{
 					EdgeUpdateTuple&  EdgeUpdate = VertexUpdateArray[i];
+					
 
-					// New vertex values 
-					const MeshVertType&  VertAttributes = EdgeUpdate.Get<2>();
+					// New vertex attribute values for the collapsed vertex 
+					const MeshVertType&  CollapsedAttributes = EdgeUpdate.Get<2>();
 
 					// update the first edge vertex
 					SimpVertType* VtxPtr = EdgeUpdate.Get<0>();
 					if (VtxPtr != nullptr)
 					{
-						MeshManager.UpdateVertexAttributes(*VtxPtr, VertAttributes);
+						MeshManager.UpdateVertexAttributes(*VtxPtr, CollapsedAttributes);
 					}
 
 					// update the second edge vertex
 					VtxPtr = EdgeUpdate.Get<1>();
 					if (VtxPtr != nullptr)
 					{
-						MeshManager.UpdateVertexAttributes(*VtxPtr, VertAttributes);
+						MeshManager.UpdateVertexAttributes(*VtxPtr, CollapsedAttributes);
 					}
 
 				}
diff --git a/Engine/Plugins/Experimental/SkeletalReduction/Source/Private/SkeletalSimplifierMeshManager.cpp b/Engine/Plugins/Experimental/SkeletalReduction/Source/Private/SkeletalSimplifierMeshManager.cpp
--- a/Engine/Plugins/Experimental/SkeletalReduction/Source/Private/SkeletalSimplifierMeshManager.cpp
+++ b/Engine/Plugins/Experimental/SkeletalReduction/Source/Private/SkeletalSimplifierMeshManager.cpp
@@ -594,6 +594,62 @@ void SkeletalSimplifier::FSimplifierMeshManager::FlagBoxCorners(const ESimpEleme
 	if (VisitedMask) delete[] VisitedMask;
 }
 
+void SkeletalSimplifier::FSimplifierMeshManager::FlagEdges(const TFunction<bool(const SimpVertType*, const SimpVertType*)> IsDifferent, const ESimpElementFlags Flag)
+{
+	int32 NumEdges = EdgeArray.Num();
+
+	for (int32 i = 0; i < NumEdges; ++i)
+	{
+		SimpEdgeType& Edge = EdgeArray[i];
+		if (IsDifferent(Edge.v0, Edge.v1))
+		{
+			Edge.EnableFlags(Flag);
+			Edge.v0->EnableFlags(Flag);
+			Edge.v1->EnableFlags(Flag);
+		}
+	}
+
+	for (int i = 0; i < NumSrcVerts; i++)
+	{
+		SimpVertType* v = &VertArray[i];
+		SimpVertType* v1 = v;
+
+		if (v1 == nullptr || v1->TestFlags(Flag))
+		{
+			continue;
+		}
+		// only one vert in this group
+		if (v1->next == v1)
+		{
+			continue;
+		}
+		
+		bool bAddedFlag = false;
+		do {
+
+			if (IsDifferent(v, v->next))
+			{
+				// we only need to mark one of the vertices in this vertex group since the lock state will be propagated to the others.
+				v->EnableFlags(Flag);
+				bAddedFlag = true;
+			}
+			v = v->next;
+		} while (v != v1);
+
+		// add the locked state to all the vertices in this group.
+		if (bAddedFlag)
+		{
+			v = v1;
+			do {
+				v->EnableFlags(Flag);
+				v = v->next;
+			} while (v != v1);
+
+		}
+		
+	}
+}
+
 void SkeletalSimplifier::FSimplifierMeshManager::GetAdjacentTopology(const SimpVertType* VertPtr,
 	TriPtrArray& DirtyTris, VertPtrArray& DirtyVerts, EdgePtrArray& DirtyEdges)
 {
diff --git a/Engine/Plugins/Experimental/SkeletalReduction/Source/Private/SkeletalSimplifierMeshManager.h b/Engine/Plugins/Experimental/SkeletalReduction/Source/Private/SkeletalSimplifierMeshManager.h
--- a/Engine/Plugins/Experimental/SkeletalReduction/Source/Private/SkeletalSimplifierMeshManager.h
+++ b/Engine/Plugins/Experimental/SkeletalReduction/Source/Private/SkeletalSimplifierMeshManager.h
@@ -89,6 +89,8 @@ namespace SkeletalSimplifier
 		// Apply a flag to all verts that are identified as being at the corner of a box.
 		void        FlagBoxCorners(const ESimpElementFlags Flag);
 
+		// Apply flag to edges when the IsDifferent(AVert, BVert) == true
+		void        FlagEdges(const TFunction<bool(const SimpVertType*, const SimpVertType*)> IsDifferent, const ESimpElementFlags Flag);
 
 		// Change the attributes on a given simplifier vert.
 		void UpdateVertexAttributes(SimpVertType& Vertex, const MeshVertType& AttributeVert)
diff --git a/Engine/Source/Editor/DetailCustomizations/Private/SkeletalMeshReductionSettingsDetails.cpp b/Engine/Source/Editor/DetailCustomizations/Private/SkeletalMeshReductionSettingsDetails.cpp
--- a/Engine/Source/Editor/DetailCustomizations/Private/SkeletalMeshReductionSettingsDetails.cpp
+++ b/Engine/Source/Editor/DetailCustomizations/Private/SkeletalMeshReductionSettingsDetails.cpp
@@ -177,7 +177,8 @@ void FSkeletalMeshReductionSettingsDetails::CustomizeChildren(TSharedRef<IProper
 			GET_MEMBER_NAME_CHECKED(FSkeletalMeshOptimizationSettings, TerminationCriterion),
 			GET_MEMBER_NAME_CHECKED(FSkeletalMeshOptimizationSettings, bLockEdges),
 			GET_MEMBER_NAME_CHECKED(FSkeletalMeshOptimizationSettings, bEnforceBoneBoundaries),
-			GET_MEMBER_NAME_CHECKED(FSkeletalMeshOptimizationSettings, VolumeImportance)
+			GET_MEMBER_NAME_CHECKED(FSkeletalMeshOptimizationSettings, VolumeImportance),
+			GET_MEMBER_NAME_CHECKED(FSkeletalMeshOptimizationSettings, bLockColorBounaries)
 		};
 
 		uint32 NumChildren = 0;
diff --git a/Engine/Source/Editor/Persona/Private/PersonaMeshDetails.cpp b/Engine/Source/Editor/Persona/Private/PersonaMeshDetails.cpp
--- a/Engine/Source/Editor/Persona/Private/PersonaMeshDetails.cpp
+++ b/Engine/Source/Editor/Persona/Private/PersonaMeshDetails.cpp
@@ -718,6 +718,13 @@ void FSkeletalMeshReductionSettingsLayout::GenerateChildContent(IDetailChildrenB
 			LOCTEXT("LockEdges_RowNameContentTooltip", "Preserve cuts in the mesh surface by locking vertices in place.  Increases the quality of the simplified mesh at edges at the cost of more triangles."),
 			FGetCheckBoxStateDelegate::CreateRaw(this, &FSkeletalMeshReductionSettingsLayout::GetLockEdges),
 			FSetCheckBoxStateDelegate::CreateRaw(this, &FSkeletalMeshReductionSettingsLayout::SetLockEdges));
+
+		AddBoolRow(ChildrenBuilder,
+			LOCTEXT("LockColorBoundaries_Row", "LockColorBoundaries"),
+			LOCTEXT("LockColorBoundaries_RowNameContent", "Lock Vertex Color Boundaries"),
+			LOCTEXT("LockColorBoundaries_RowNameContentTooltip", "Locking edges that connect two vertex colors.  Increases the quality of the simplified mesh at edges at the cost of more triangles."),
+			FGetCheckBoxStateDelegate::CreateRaw(this, &FSkeletalMeshReductionSettingsLayout::GetLockColorBounaries),
+			FSetCheckBoxStateDelegate::CreateRaw(this, &FSkeletalMeshReductionSettingsLayout::SetLockColorBounaries));
 	}
 
 	AddBaseLODRow(ChildrenBuilder);
@@ -1208,6 +1215,20 @@ void FSkeletalMeshReductionSettingsLayout::SetLockEdges(ECheckBoxState NewState)
 	ReductionSettings.bLockEdges = (NewState == ECheckBoxState::Checked) ? true : false;
 }
 
+ECheckBoxState FSkeletalMeshReductionSettingsLayout::GetLockColorBounaries() const
+{
+	return ReductionSettings.bLockColorBounaries ? ECheckBoxState::Checked : ECheckBoxState::Unchecked;
+}
+
+void FSkeletalMeshReductionSettingsLayout::SetLockColorBounaries(ECheckBoxState NewState)
+{
+	FText TransactionText = FText::Format(LOCTEXT("PersonaReductionChangedSetLockColorBounariesLOD", "LOD{0} reduction settings: lock color boundaries changed"), LODIndex);
+	FScopedTransaction Transaction(TransactionText);
+	ModifyMeshLODSettingsDelegate.ExecuteIfBound(LODIndex);
+
+	ReductionSettings.bLockColorBounaries = (NewState == ECheckBoxState::Checked) ? true : false;
+}
+
 ECheckBoxState FSkeletalMeshReductionSettingsLayout::GetEnforceBoneBoundaries() const
 {
 	return ReductionSettings.bEnforceBoneBoundaries ? ECheckBoxState::Checked : ECheckBoxState::Unchecked;
diff --git a/Engine/Source/Editor/Persona/Private/PersonaMeshDetails.h b/Engine/Source/Editor/Persona/Private/PersonaMeshDetails.h
--- a/Engine/Source/Editor/Persona/Private/PersonaMeshDetails.h
+++ b/Engine/Source/Editor/Persona/Private/PersonaMeshDetails.h
@@ -172,6 +172,9 @@ private:
 	ECheckBoxState GetLockEdges() const;
 	void SetLockEdges(ECheckBoxState NewState);
 
+	ECheckBoxState GetLockColorBounaries() const;
+	void SetLockColorBounaries(ECheckBoxState NewState);
+
 	ECheckBoxState GetEnforceBoneBoundaries() const;
 	void SetEnforceBoneBoundaries(ECheckBoxState NewState);
 
diff --git a/Engine/Source/Runtime/Engine/Private/SkeletalMesh.cpp b/Engine/Source/Runtime/Engine/Private/SkeletalMesh.cpp
--- a/Engine/Source/Runtime/Engine/Private/SkeletalMesh.cpp
+++ b/Engine/Source/Runtime/Engine/Private/SkeletalMesh.cpp
@@ -4003,6 +4003,7 @@ static void SerializeReductionSettingsForDDC(FArchive& Ar, FSkeletalMeshOptimiza
 	FArchive_Serialize_BitfieldBool(Ar, ReductionSettings.bRecalcNormals);
 	FArchive_Serialize_BitfieldBool(Ar, ReductionSettings.bEnforceBoneBoundaries);
 	FArchive_Serialize_BitfieldBool(Ar, ReductionSettings.bLockEdges);
+	FArchive_Serialize_BitfieldBool(Ar, ReductionSettings.bLockColorBounaries);
 	Ar << ReductionSettings.TerminationCriterion;
 	Ar << ReductionSettings.NumOfTrianglesPercentage;
 	Ar << ReductionSettings.NumOfVertPercentage;
diff --git a/Engine/Source/Runtime/Engine/Public/SkeletalMeshReductionSettings.h b/Engine/Source/Runtime/Engine/Public/SkeletalMeshReductionSettings.h
--- a/Engine/Source/Runtime/Engine/Public/SkeletalMeshReductionSettings.h
+++ b/Engine/Source/Runtime/Engine/Public/SkeletalMeshReductionSettings.h
@@ -138,6 +138,10 @@ struct FSkeletalMeshOptimizationSettings
 	UPROPERTY(EditAnywhere, Category = FSkeletalMeshOptimizationSettings, meta = (DisplayName = "Lock Mesh Edges"))
 	uint8 bLockEdges : 1;
 
+	/** Disallow edge collapse when the vertices do not have a common color*/
+	UPROPERTY(EditAnywhere, Category = FSkeletalMeshOptimizationSettings, meta = (DisplayName = "Lock Vertex Color Boundaries"))
+	uint8 bLockColorBounaries : 1;
+
 	/** Base LOD index to generate this LOD. By default, we generate from LOD 0 */
 	UPROPERTY(EditAnywhere, Category = FSkeletalMeshOptimizationSettings)
 	int32 BaseLOD;
@@ -173,6 +177,7 @@ struct FSkeletalMeshOptimizationSettings
 		, bEnforceBoneBoundaries(false)
 		, VolumeImportance(1.f)
 		, bLockEdges(false)
+		, bLockColorBounaries(false)
 		, BaseLOD(0)
 #if WITH_EDITORONLY_DATA
 		, BakePose_DEPRECATED(nullptr)
