Fix reduce LOD containing disabled section and that was reduce before we get rid of disabled section in the reduction.
#rnx

diff --git a/Engine/Plugins/Experimental/SkeletalReduction/Source/Private/SkeletalMeshReductionPlugin.cpp b/Engine/Plugins/Experimental/SkeletalReduction/Source/Private/SkeletalMeshReductionPlugin.cpp
--- a/Engine/Plugins/Experimental/SkeletalReduction/Source/Private/SkeletalMeshReductionPlugin.cpp
+++ b/Engine/Plugins/Experimental/SkeletalReduction/Source/Private/SkeletalMeshReductionPlugin.cpp
@@ -2083,6 +2083,16 @@ void FQuadricSkeletalMeshReduction::ReduceSkeletalMesh(USkeletalMesh& SkeletalMe
 						Section.ChunkedParentSectionIndex = bIsChunkedSection ? CurrentParentSectionIndex : INDEX_NONE;
 						//If we reduce inline the source model, we want to use the real source original section
 						Section.OriginalDataSectionIndex = bReducingSourceModel ? SectionData.OriginalDataSectionIndex : OriginalSectionIndex;
+
+						if (!bLODModelAdded)
+						{
+							if (FSkelMeshSourceSectionUserData* BackUpUserSectionData = BackupUserSectionsData.Find(SectionData.OriginalDataSectionIndex))
+							{
+								FSkelMeshSourceSectionUserData& ReducedUserSectionData = ImportedModelLOD.UserSectionsData.FindOrAdd(Section.OriginalDataSectionIndex);
+								ReducedUserSectionData = *BackUpUserSectionData;
+							}
+						}
+
 						SectionMatched = true; //a backup section can be restore only once
 						break;
 					}
@@ -2091,8 +2101,7 @@ void FQuadricSkeletalMeshReduction::ReduceSkeletalMesh(USkeletalMesh& SkeletalMe
 
 			if (!bLODModelAdded)
 			{
-				//If its an existing LOD re-apply the UserSectionData
-				ImportedModelLOD.UserSectionsData = BackupUserSectionsData;
+				//If its an existing LOD put back the buildStringID
 				ImportedModelLOD.BuildStringID = BackupLodModelBuildStringID;
 			}
 		}
