Skeletal Mesh Reduction. Fix Static Analysis warnings.

diff --git a/Engine/Plugins/Experimental/SkeletalReduction/Source/Private/SkeletalMeshReductionPlugin.cpp b/Engine/Plugins/Experimental/SkeletalReduction/Source/Private/SkeletalMeshReductionPlugin.cpp
--- a/Engine/Plugins/Experimental/SkeletalReduction/Source/Private/SkeletalMeshReductionPlugin.cpp
+++ b/Engine/Plugins/Experimental/SkeletalReduction/Source/Private/SkeletalMeshReductionPlugin.cpp
@@ -1455,7 +1455,8 @@ void  FQuadricSkeletalMeshReduction::AddSourceModelInfluences( const FSkeletalMe
 
 			//  Create a structure that allows us to look-up by SourceModel Vertex ID 
 			
-			BoneIdWeightMap* VtxToBoneIdWeightMap             = new  BoneIdWeightMap[NumImportVertex];
+			TArray<BoneIdWeightMap> VtxToBoneIdWeightMap;
+			VtxToBoneIdWeightMap.AddDefaulted(NumImportVertex);
 
 			for (int32 i = 0; i < SrcModelInfluences.Num(); ++i)
 			{
@@ -1541,7 +1542,6 @@ void  FQuadricSkeletalMeshReduction::AddSourceModelInfluences( const FSkeletalMe
 				}
 			}
 
-			delete[] VtxToBoneIdWeightMap;
 		}
 
 	}
diff --git a/Engine/Plugins/Experimental/SkeletalReduction/Source/Private/SkeletalSimplifierMeshManager.cpp b/Engine/Plugins/Experimental/SkeletalReduction/Source/Private/SkeletalSimplifierMeshManager.cpp
--- a/Engine/Plugins/Experimental/SkeletalReduction/Source/Private/SkeletalSimplifierMeshManager.cpp
+++ b/Engine/Plugins/Experimental/SkeletalReduction/Source/Private/SkeletalSimplifierMeshManager.cpp
@@ -22,10 +22,11 @@ SkeletalSimplifier::FSimplifierMeshManager::FSimplifierMeshManager(const MeshVer
 	EdgeVertIdHashMap(1 << FMath::Min(16u, FMath::FloorLog2(InNumSrcVerts)))
 {
 	// Allocate verts and tris
+	VertArray.Empty();
+	VertArray.AddDefaulted(NumSrcVerts);
 
-	VertArray = new SimpVertType[NumSrcVerts];
-	TriArray  = new SimpTriType[NumSrcTris];
-
+	TriArray.Empty();
+	TriArray.AddDefaulted(NumSrcTris);
 
 	// Deep copy the verts
 
@@ -68,11 +69,11 @@ SkeletalSimplifier::FSimplifierMeshManager::FSimplifierMeshManager(const MeshVer
 
 	// Group the verts that share the same location.
 
-	GroupVerts(VertArray, NumSrcVerts);
+	GroupVerts(VertArray);
 
 	// Populate EdgeArray
 
-	MakeEdges(VertArray, NumSrcVerts, NumSrcTris, EdgeArray);
+	MakeEdges(VertArray, NumSrcTris, EdgeArray);
 
 	// Links all the edges together
 	GroupEdges(EdgeArray);
@@ -100,8 +101,10 @@ SkeletalSimplifier::FSimplifierMeshManager::FSimplifierMeshManager(const MeshVer
 
 }
 
-void SkeletalSimplifier::FSimplifierMeshManager::GroupVerts(SimpVertType* Verts, const int32 NumVerts)
+void SkeletalSimplifier::FSimplifierMeshManager::GroupVerts(TArray<SimpVertType>& Verts)
 {
+	const int32 NumVerts = Verts.Num();
+
 	// group verts that share a point
 	FHashTable HashTable(1 << FMath::Min(16u, FMath::FloorLog2(NumVerts / 2)), NumVerts);
 
@@ -164,10 +167,11 @@ void SkeletalSimplifier::FSimplifierMeshManager::GroupVerts(SimpVertType* Verts,
 }
 
 
-void SkeletalSimplifier::FSimplifierMeshManager::MakeEdges(const SimpVertType* Verts, const int32 NumVerts, const int32 NumTris, TArray<SimpEdgeType>& Edges)
+void SkeletalSimplifier::FSimplifierMeshManager::MakeEdges(const TArray<SimpVertType>& Verts, const int32 NumTris, TArray<SimpEdgeType>& Edges)
 {
 
 	// Populate the TArray of edges.
+	const int32 NumVerts = Verts.Num();
 
 	int32 maxEdgeSize = FMath::Min(3 * NumTris, 3 * NumVerts - 6);
 	Edges.Empty(maxEdgeSize);
@@ -505,12 +509,13 @@ void SkeletalSimplifier::FSimplifierMeshManager::FlagBoxCorners(const ESimpEleme
 {
 
 
-	uint32* VisitedMask = new uint32[NumSrcVerts](); // initializes to zeros.
+	TArray<int32> VisitedMask;
+	VisitedMask.AddZeroed(NumSrcVerts);
 
 	for (int32 i = 0; i < NumSrcVerts; i++)
 	{
 
-		if (VisitedMask[i])
+		if (VisitedMask[i] != 0)
 		{
 			continue;
 		}
@@ -591,7 +596,6 @@ void SkeletalSimplifier::FSimplifierMeshManager::FlagBoxCorners(const ESimpEleme
 	
 	}
 
-	if (VisitedMask) delete[] VisitedMask;
 }
 
 void SkeletalSimplifier::FSimplifierMeshManager::FlagEdges(const TFunction<bool(const SimpVertType*, const SimpVertType*)> IsDifferent, const ESimpElementFlags Flag)
@@ -1267,7 +1271,7 @@ int32 SkeletalSimplifier::FSimplifierMeshManager::CountDegeneratesTris() const
 	// not sure why this happens
 	for (int i = 0; i < NumSrcTris; i++)
 	{
-		SimpTriType* tri = &TriArray[i];
+		const SimpTriType* tri = &TriArray[i];
 
 		if (tri->TestFlags(SIMP_REMOVED))
 			continue;
diff --git a/Engine/Plugins/Experimental/SkeletalReduction/Source/Private/SkeletalSimplifierMeshManager.h b/Engine/Plugins/Experimental/SkeletalReduction/Source/Private/SkeletalSimplifierMeshManager.h
--- a/Engine/Plugins/Experimental/SkeletalReduction/Source/Private/SkeletalSimplifierMeshManager.h
+++ b/Engine/Plugins/Experimental/SkeletalReduction/Source/Private/SkeletalSimplifierMeshManager.h
@@ -73,8 +73,6 @@ namespace SkeletalSimplifier
 
 		~FSimplifierMeshManager()
 		{
-			if (VertArray) delete[] VertArray;
-			if (TriArray)  delete[] TriArray;
 		}
 
 		// Extract the currently valid verts / indices from this object.  If LockedVerts != NULL
@@ -425,8 +423,11 @@ namespace SkeletalSimplifier
 		int32   ReducedNumVerts;
 		int32   ReducedNumTris;
 
-		SimpVertType*		VertArray;
-		SimpTriType*		TriArray;
+		// Note after these arrays are constructed, they should never be resized.
+		// code holds pointers to array elements.
+		TArray<SimpVertType>  VertArray;
+		TArray<SimpTriType>   TriArray;
+
 
 		// Hash based on the Ids of the edge's verts.
 		// used to map verts to edges.
@@ -441,8 +442,8 @@ namespace SkeletalSimplifier
 
 		// Methods used in the initial construction of the simplifier mesh
 
-		void GroupVerts(SimpVertType* Verts, const int32 NumVerts);
-		void MakeEdges(const SimpVertType* Verts, const int32 NumVerts, const int32 NumTris, TArray<SimpEdgeType>& Edges);
+		void GroupVerts(TArray<SimpVertType>& Verts);
+		void MakeEdges(const TArray<SimpVertType>& Verts, const int32 NumTris, TArray<SimpEdgeType>& Edges);
 		void AppendConnectedEdges(const SimpVertType* Vert, TArray<SimpEdgeType>& Edges);
 		void GroupEdges(TArray< SimpEdgeType >& Edges);
 
diff --git a/Engine/Plugins/Experimental/SkeletalReduction/Source/Private/SkeletalSimplifierQuadricCache.h b/Engine/Plugins/Experimental/SkeletalReduction/Source/Private/SkeletalSimplifierQuadricCache.h
--- a/Engine/Plugins/Experimental/SkeletalReduction/Source/Private/SkeletalSimplifierQuadricCache.h
+++ b/Engine/Plugins/Experimental/SkeletalReduction/Source/Private/SkeletalSimplifierQuadricCache.h
@@ -28,7 +28,7 @@ namespace SkeletalSimplifier
 		*/
 		void RegisterMesh(const FSimplifierMeshManager& Mesh)
 		{
-			RegisterCache(Mesh.VertArray, Mesh.NumSrcVerts, Mesh.TriArray, Mesh.NumSrcTris);
+			RegisterCache(Mesh.VertArray, Mesh.TriArray);
 		}
 
 		/**
@@ -88,12 +88,10 @@ namespace SkeletalSimplifier
 		* Associate the cache with the simplifier vert and triangle arrays.
 		* This must be done before the cache can be used.
 		*
-		* @param VertOffset  Pointer to the array of verts in the source mesh
-		* @param NumVerts    Number of Verts in the source mesh
-		* @param TriOffset   Pointer to the array of Tris in the source mesh
-		* @param NumTris     Number of tris in the source mesh
+		* @param VertArray   Array of verts in the source mesh
+		* @param TriArray    Array of Tris in the source mesh
 		*/
-		void RegisterCache(const SimpVertType* VertOffset, const int32 NumVerts, const SimpTriType* TriOffset, const int32 NumTris);
+		void RegisterCache(const TArray<SimpVertType>& VertArray, const TArray<SimpTriType>& TriArray);
 
 		/**
 		* Allocate the member arrays with the size required
@@ -148,10 +146,13 @@ namespace SkeletalSimplifier
 
 
 	template <typename WedgeQuadricType>
-	void TQuadricCache< WedgeQuadricType>::RegisterCache(const SimpVertType* VertOffset, const int32 NumVerts, const SimpTriType* TriOffset, const int32 NumTris)
+	void TQuadricCache< WedgeQuadricType>::RegisterCache(const TArray<SimpVertType>& VertArray,  const TArray<SimpTriType>& TriArray )
 	{
-		sVerts = VertOffset;
-		sTris = TriOffset;
+		const int32 NumVerts = VertArray.Num();
+		const int32 NumTris  = TriArray.Num();
+	
+		sVerts = VertArray.GetData();
+		sTris  = TriArray.GetData();
 
 		AllocateCacheArrays(NumVerts, NumTris);
 	}
