Make sure old reduced skeletalmesh asset with disabled section keep the disable section state when restoring after the re-import.

diff --git a/Engine/Plugins/Experimental/SkeletalReduction/Source/Private/SkeletalMeshReductionPlugin.cpp b/Engine/Plugins/Experimental/SkeletalReduction/Source/Private/SkeletalMeshReductionPlugin.cpp
--- a/Engine/Plugins/Experimental/SkeletalReduction/Source/Private/SkeletalMeshReductionPlugin.cpp
+++ b/Engine/Plugins/Experimental/SkeletalReduction/Source/Private/SkeletalMeshReductionPlugin.cpp
@@ -1777,6 +1777,13 @@ void FQuadricSkeletalMeshReduction::ReduceSkeletalMesh(USkeletalMesh& SkeletalMe
 	const FSkeletalMeshOptimizationSettings& Settings = LODInfo->ReductionSettings;
 	
 
+	FSkeletalMeshLODModel DstModelBackup;
+	if (!bLODModelAdded)
+	{
+		DstModelBackup.Sections = SkeletalMeshResource.LODModels[LODIndex].Sections;
+		DstModelBackup.UserSectionsData = SkeletalMeshResource.LODModels[LODIndex].UserSectionsData;
+	}
+
 	// Struct to identify important bones.  Vertices associated with these bones
 	// will have additional collapse weight added to them.
 
@@ -1820,54 +1827,52 @@ void FQuadricSkeletalMeshReduction::ReduceSkeletalMesh(USkeletalMesh& SkeletalMe
 		}
 	}
 
-	//Store the sections flags
-	struct FSectionData
-	{
-		uint16 MaterialIndex;
-		int32 MaterialMap;
-		bool bCastShadow;
-		bool bRecomputeTangent;
-		bool bDisabled;
-		int32 GenerateUpToLodIndex;
-		int32 ChunkedParentSectionIndex;
-		int32 OriginalDataSectionIndex;
-	};
-
-	TMap<int32, FSkelMeshSourceSectionUserData> BackupUserSectionsData;
-	FString BackupLodModelBuildStringID = TEXT("");
+	//We backup only the sections and the user sections data
+	FSkeletalMeshLODModel SrcModelBackup;
+	SrcModelBackup.Sections = SrcModel->Sections;
+	SrcModelBackup.UserSectionsData = SrcModel->UserSectionsData;
 
-
-	
-
-	auto FillSectionMaterialSlot = [&SkeletalMeshResource, &LODIndex, bLODModelAdded](TArray<int32>& SectionMaterialSlot)
+	//Restore the source sections data
+	auto RestoreUserSectionsData = [](const FSkeletalMeshLODModel& SourceLODModel, FSkeletalMeshLODModel& DestinationLODModel)
 	{
-		SectionMaterialSlot.Empty();
-		if (!bLODModelAdded && SkeletalMeshResource.LODModels.IsValidIndex(LODIndex))
+		//Now restore the reduce section user change and adjust the originalDataSectionIndex to point on the correct UserSectionData
+		TBitArray<> SourceSectionMatched;
+		SourceSectionMatched.Init(false, SourceLODModel.Sections.Num());
+		for (int32 SectionIndex = 0; SectionIndex < DestinationLODModel.Sections.Num(); ++SectionIndex)
 		{
-			int32 SectionNumber = SkeletalMeshResource.LODModels[LODIndex].Sections.Num();
-			for (int32 SectionIndex = 0; SectionIndex < SectionNumber; ++SectionIndex)
+			FSkelMeshSection& Section = DestinationLODModel.Sections[SectionIndex];
+			FSkelMeshSourceSectionUserData& DestinationUserData = FSkelMeshSourceSectionUserData::GetSourceSectionUserData(DestinationLODModel.UserSectionsData, Section);
+			for (int32 SourceSectionIndex = 0; SourceSectionIndex < SourceLODModel.Sections.Num(); ++SourceSectionIndex)
 			{
-				SectionMaterialSlot.Add(SkeletalMeshResource.LODModels[LODIndex].Sections[SectionIndex].OriginalDataSectionIndex);
+				if (SourceSectionMatched[SourceSectionIndex])
+				{
+					continue;
+				}
+				const FSkelMeshSection& SourceSection = SourceLODModel.Sections[SourceSectionIndex];
+				if (const FSkelMeshSourceSectionUserData* SourceUserData = SourceLODModel.UserSectionsData.Find(SourceSection.OriginalDataSectionIndex))
+				{
+					if (Section.MaterialIndex == SourceSection.MaterialIndex)
+					{
+						DestinationUserData = *SourceUserData;
+						SourceSectionMatched[SourceSectionIndex] = true;
+						break;
+					}
+				}
 			}
 		}
+		DestinationLODModel.SyncronizeUserSectionsDataArray();
 	};
 
+	FString BackupLodModelBuildStringID = TEXT("");
+
 	// Unbind any existing clothing assets before we reimport the geometry
 	TArray<ClothingAssetUtils::FClothingAssetMeshBinding> ClothingBindings;
-	//Get a map of enable/disable sections
-	TArray<int32> OriginalSectionMaterialSlot;
-
 	//Do not play with cloth if the LOD is added
 	if (!bLODModelAdded)
 	{
 		FLODUtilities::UnbindClothingAndBackup(&SkeletalMesh, ClothingBindings, LODIndex);
-	}
-
-	if (!bLODModelAdded)
-	{
-		FSkeletalMeshLODModel& DstBackupSectionLODModel = SkeletalMeshResource.LODModels[LODIndex];
-		BackupLodModelBuildStringID = DstBackupSectionLODModel.BuildStringID;
-		BackupUserSectionsData = DstBackupSectionLODModel.UserSectionsData;
+		//We have to put back the exact UserSectionsData to not invalidate the DDC key
+		BackupLodModelBuildStringID = SkeletalMeshResource.LODModels[LODIndex].BuildStringID;
 	}
 
 	bool bReducingSourceModel = false;
@@ -1876,6 +1881,9 @@ void FQuadricSkeletalMeshReduction::ReduceSkeletalMesh(USkeletalMesh& SkeletalMe
 	{
 		TMap<FString, TArray<FMorphTargetDelta>> TempLODMorphTargetData;
 		SkelResource->OriginalReductionSourceMeshData[BaseLOD]->LoadReductionData(*SrcModel, TempLODMorphTargetData, &SkeletalMesh);
+
+		//Restore the section data state (like disabled...)
+		RestoreUserSectionsData(SrcModelBackup, *SrcModel);
 		//Rebackup the source model since we update it, source always have empty LODMaterial map
 		//If you swap a material ID and after you do inline reduction, you have to remap it again, but not if you reduce and then remap the materialID
 		//this is by design currently
@@ -1888,31 +1896,6 @@ void FQuadricSkeletalMeshReduction::ReduceSkeletalMesh(USkeletalMesh& SkeletalMe
 
 	check(SrcModel);
 
-	//We backup the section data to put keep the LODModel Section in a good state after the reduction
-	TMap<int32, FSectionData> SrcBackupSectionIndexToSectionData;
-	{
-		FSkeletalMeshLODInfo* BackupLODInfo = SkeletalMesh.GetLODInfo(Settings.BaseLOD);
-		FSkeletalMeshLODModel& LODModelToBackup = *SrcModel;
-		SrcBackupSectionIndexToSectionData.Empty(LODModelToBackup.Sections.Num());
-		for (int32 SectionIndex = 0; SectionIndex < LODModelToBackup.Sections.Num(); ++SectionIndex)
-		{
-			FSectionData& SectionData = SrcBackupSectionIndexToSectionData.FindOrAdd(SectionIndex);
-			SectionData.MaterialIndex = LODModelToBackup.Sections[SectionIndex].MaterialIndex;
-			SectionData.bCastShadow = LODModelToBackup.Sections[SectionIndex].bCastShadow;
-			SectionData.bRecomputeTangent = LODModelToBackup.Sections[SectionIndex].bRecomputeTangent;
-			SectionData.bDisabled = LODModelToBackup.Sections[SectionIndex].bDisabled;
-			SectionData.GenerateUpToLodIndex = LODModelToBackup.Sections[SectionIndex].GenerateUpToLodIndex;
-			SectionData.ChunkedParentSectionIndex = LODModelToBackup.Sections[SectionIndex].ChunkedParentSectionIndex;
-			SectionData.OriginalDataSectionIndex = LODModelToBackup.Sections[SectionIndex].OriginalDataSectionIndex;
-			SectionData.MaterialMap = (BackupLODInfo != nullptr && BackupLODInfo->LODMaterialMap.IsValidIndex(SectionIndex)) ? BackupLODInfo->LODMaterialMap[SectionIndex] : INDEX_NONE;
-			if (SectionData.MaterialMap == SectionData.MaterialIndex)
-			{
-				//Remove any override if the value is the same
-				SectionData.MaterialMap = INDEX_NONE;
-			}
-		}
-	}
-
 	// now try bone reduction process if it's setup
 	TMap<FBoneIndexType, FBoneIndexType> BonesToRemove;
 	
@@ -1985,6 +1968,7 @@ void FQuadricSkeletalMeshReduction::ReduceSkeletalMesh(USkeletalMesh& SkeletalMe
 
 	// Swap out the old model.  
 	FSkeletalMeshImportData RawMesh;
+	bool bPutBackRawMesh = false;
 	ESkeletalMeshGeoImportVersions GeoImportVersion = ESkeletalMeshGeoImportVersions::Before_Versionning;
 	ESkeletalMeshSkinningImportVersions SkinningImportVersion = ESkeletalMeshSkinningImportVersions::Before_Versionning;
 	{
@@ -2000,6 +1984,7 @@ void FQuadricSkeletalMeshReduction::ReduceSkeletalMesh(USkeletalMesh& SkeletalMe
 			{
 				SkeletalMesh.LoadLODImportedData(LODIndex, RawMesh);
 				SkeletalMesh.GetLODImportedDataVersions(LODIndex, GeoImportVersion, SkinningImportVersion);
+				bPutBackRawMesh = true;
 			}
 			//If the delegate is not bound 
 			if (!Settings.OnDeleteLODModelDelegate.IsBound())
@@ -2018,13 +2003,11 @@ void FQuadricSkeletalMeshReduction::ReduceSkeletalMesh(USkeletalMesh& SkeletalMe
 		{
 			SkeletalMesh.LoadLODImportedData(BaseLOD, RawMesh);
 			SkeletalMesh.GetLODImportedDataVersions(BaseLOD, GeoImportVersion, SkinningImportVersion);
+			bPutBackRawMesh = true;
 		}
 	}
 
-	
-
 	// Reduce LOD model with SrcMesh
-
 	if (ReduceSkeletalLODModel(*SrcModel, *NewModel, SkeletalMesh.GetImportedBounds(), SkeletalMesh.RefSkeleton, Settings, ImportantBones, RelativeToRefPoseMatrices, LODIndex, bReducingSourceModel))
 	{
 		FSkeletalMeshLODInfo* ReducedLODInfoPtr = SkeletalMesh.GetLODInfo(LODIndex);
@@ -2047,60 +2030,41 @@ void FQuadricSkeletalMeshReduction::ReduceSkeletalMesh(USkeletalMesh& SkeletalMe
 		// Flag this LOD as having been simplified.
 		ReducedLODInfoPtr->bHasBeenSimplified = true;
 		SkeletalMesh.bHasBeenSimplified = true;
-		
-		//Restore the source sections data
+
+		//Restore the user sections data to what it was. It must be done if we want to avoid changing the DDC key. I.E. UserSectionData is part of the key
+		//DDC key cannot be change during the build
 		{
 			FSkeletalMeshLODModel& ImportedModelLOD = SkeletalMesh.GetImportedModel()->LODModels[LODIndex];
-			TMap<int32, bool> OriginalSectionMatched;
-			OriginalSectionMatched.Reserve(SrcBackupSectionIndexToSectionData.Num());
-			int32 CurrentParentSectionIndex = INDEX_NONE;
-			int32 OriginalSectionIndex = INDEX_NONE;
-			for (int32 SectionIndex = 0; SectionIndex < ImportedModelLOD.Sections.Num(); ++SectionIndex)
+			RestoreUserSectionsData(SrcModelBackup, ImportedModelLOD);
+
+			if (!bLODModelAdded)
 			{
-				FSkelMeshSection& Section = ImportedModelLOD.Sections[SectionIndex];
-				for (auto Kvp : SrcBackupSectionIndexToSectionData)
+				TBitArray<> SourceSectionMatched;
+				SourceSectionMatched.Init(false, DstModelBackup.Sections.Num());
+				for (int32 SectionIndex = 0; SectionIndex < ImportedModelLOD.Sections.Num(); ++SectionIndex)
 				{
-					const int32 SourceSectionIndex = Kvp.Key;
-					bool& SectionMatched = OriginalSectionMatched.FindOrAdd(SourceSectionIndex);
-					if (SectionMatched)
-					{
-						continue;
-					}
-					const FSectionData& SectionData = Kvp.Value;
-					//We use the material index to match the section
-					if (Section.MaterialIndex == SectionData.MaterialIndex)
+					FSkelMeshSection& Section = ImportedModelLOD.Sections[SectionIndex];
+					for (int32 SourceSectionIndex = 0; SourceSectionIndex < DstModelBackup.Sections.Num(); ++SourceSectionIndex)
 					{
-						bool bIsChunkedSection = SectionData.ChunkedParentSectionIndex != INDEX_NONE;
-						if (!bIsChunkedSection)
+						if (SourceSectionMatched[SourceSectionIndex])
 						{
-							CurrentParentSectionIndex = SectionIndex;
-							OriginalSectionIndex++;
+							continue;
 						}
-						Section.bCastShadow = SectionData.bCastShadow;
-						Section.bRecomputeTangent = SectionData.bRecomputeTangent;
-						Section.bDisabled = SectionData.bDisabled;
-						Section.GenerateUpToLodIndex = SectionData.GenerateUpToLodIndex;
-						Section.ChunkedParentSectionIndex = bIsChunkedSection ? CurrentParentSectionIndex : INDEX_NONE;
-						//If we reduce inline the source model, we want to use the real source original section
-						Section.OriginalDataSectionIndex = bReducingSourceModel ? SectionData.OriginalDataSectionIndex : OriginalSectionIndex;
-
-						if (!bLODModelAdded)
+						const FSkelMeshSection& SourceSection = DstModelBackup.Sections[SourceSectionIndex];
+						if (const FSkelMeshSourceSectionUserData* SourceUserData = DstModelBackup.UserSectionsData.Find(SourceSection.OriginalDataSectionIndex))
 						{
-							if (FSkelMeshSourceSectionUserData* BackUpUserSectionData = BackupUserSectionsData.Find(SectionData.OriginalDataSectionIndex))
+							if (Section.MaterialIndex == SourceSection.MaterialIndex)
 							{
-								FSkelMeshSourceSectionUserData& ReducedUserSectionData = ImportedModelLOD.UserSectionsData.FindOrAdd(Section.OriginalDataSectionIndex);
-								ReducedUserSectionData = *BackUpUserSectionData;
+								Section.OriginalDataSectionIndex = SourceSection.OriginalDataSectionIndex;
+								SourceSectionMatched[SourceSectionIndex] = true;
+								break;
 							}
 						}
-
-						SectionMatched = true; //a backup section can be restore only once
-						break;
 					}
 				}
-			}
 
-			if (!bLODModelAdded)
-			{
+				ImportedModelLOD.UserSectionsData = DstModelBackup.UserSectionsData;
+				ImportedModelLOD.SyncronizeUserSectionsDataArray();
 				//If its an existing LOD put back the buildStringID
 				ImportedModelLOD.BuildStringID = BackupLodModelBuildStringID;
 			}
@@ -2108,10 +2072,10 @@ void FQuadricSkeletalMeshReduction::ReduceSkeletalMesh(USkeletalMesh& SkeletalMe
 	}
 	else
 	{
+		//TODO: This is not thread safe.
 		FSkeletalMeshLODModel::CopyStructure(NewModel, SrcModel);
 
 		// Do any joint-welding / bone removal.
-
 		if (MeshBoneReductionInterface != NULL && MeshBoneReductionInterface->GetBoneReductionData(&SkeletalMesh, LODIndex, BonesToRemove))
 		{
 			// fix up chunks to remove the bones that set to be removed
@@ -2146,20 +2110,17 @@ void FQuadricSkeletalMeshReduction::ReduceSkeletalMesh(USkeletalMesh& SkeletalMe
 	
 	if (!bLODModelAdded)
 	{
-		//Get the number of enabled section
-		TArray<int32> SectionMaterialSlotAfterReduction;
-		FillSectionMaterialSlot(SectionMaterialSlotAfterReduction);
-
 		//Put back the clothing for this newly reduce LOD
 		if (ClothingBindings.Num() > 0)
 		{
 			FLODUtilities::RestoreClothingFromBackup(&SkeletalMesh, ClothingBindings, LODIndex);
 		}
 	}
-
-	if ((bReducingSourceModel || !bLODModelAdded ) && RawMesh.Points.Num() > 0)
+	if (bPutBackRawMesh)
 	{
+		check((bReducingSourceModel || !bLODModelAdded));
 		//Put back the original import data, we need it to allow inline reduction and skeletal mesh split workflow
+		//It also warranty that we do not change the ddc key
 		SkeletalMesh.SaveLODImportedData(LODIndex, RawMesh);
 		SkeletalMesh.SetLODImportedDataVersions(LODIndex, GeoImportVersion, SkinningImportVersion);
 	}
diff --git a/Engine/Source/Developer/SkeletalMeshUtilitiesCommon/Private/LODUtilities.cpp b/Engine/Source/Developer/SkeletalMeshUtilitiesCommon/Private/LODUtilities.cpp
--- a/Engine/Source/Developer/SkeletalMeshUtilitiesCommon/Private/LODUtilities.cpp
+++ b/Engine/Source/Developer/SkeletalMeshUtilitiesCommon/Private/LODUtilities.cpp
@@ -811,6 +811,22 @@ void FLODUtilities::ApplyMorphTargetsToLOD(USkeletalMesh* SkeletalMesh, int32 So
 		//Abort remapping of morph target since the data is missing
 		return;
 	}
+
+	//Make sure we have some morph for this LOD
+	bool bContainsMorphTargets = false;
+	for (UMorphTarget* MorphTarget : SkeletalMesh->MorphTargets)
+	{
+		if (MorphTarget->HasDataForLOD(SourceLOD))
+		{
+			bContainsMorphTargets = true;
+		}
+	}
+	if (!bContainsMorphTargets)
+	{
+		//No morph target to remap
+		return;
+	}
+
 	if (bReduceBaseLOD)
 	{
 		ReductionBaseSkeletalMeshBulkData = SkeletalMeshResource->OriginalReductionSourceMeshData[SourceLOD];
@@ -851,21 +867,6 @@ void FLODUtilities::ApplyMorphTargetsToLOD(USkeletalMesh* SkeletalMesh, int32 So
 		return InternalGetSectionMaterialIndex(TargetLODModel, SectionIndex);
 	};
 
-	//Make sure we have some morph for this LOD
-	bool bContainsMorphTargets = false;
-	for (UMorphTarget *MorphTarget : SkeletalMesh->MorphTargets)
-	{
-		if (MorphTarget->HasDataForLOD(SourceLOD))
-		{
-			bContainsMorphTargets = true;
-		}
-	}
-	if (!bContainsMorphTargets)
-	{
-		//No morph target to remap
-		return;
-	}
-
 	//We have to match target sections index with the correct base section index. Reduced LODs can contain a different number of sections than the base LOD
 	TArray<int32> TargetSectionMatchBaseIndex;
 	//Initialize the array to INDEX_NONE
@@ -1161,8 +1162,7 @@ void FLODUtilities::RestoreSkeletalMeshLODImportedData(USkeletalMesh* SkeletalMe
 	{
 		if (!SkeletalMesh->IsLODImportedDataBuildAvailable(LodIndex))
 		{
-			TArray<int32> EmptyLodInfoMaterialMap;
-			ImportedBaseLODModel.UpdateChunkedSectionInfo(SkeletalMesh->GetName(), EmptyLodInfoMaterialMap);
+			ImportedBaseLODModel.UpdateChunkedSectionInfo(SkeletalMesh->GetName());
 		}
 		//When we restore a LOD we destroy the LODMaterialMap (user manual section material slot assignation)
 		FSkeletalMeshLODInfo* LODInfo = SkeletalMesh->GetLODInfo(LodIndex);
@@ -2646,4 +2646,4 @@ void FLODUtilities::RestoreClothingFromBackup(USkeletalMesh* SkeletalMesh, TArra
 
 
 
-#undef LOCTEXT_NAMESPACE // "LODUtilities"
+#undef LOCTEXT_NAMESPACE // "LODUtilities"
\ No newline at end of file
diff --git a/Engine/Source/Editor/UnrealEd/Private/Factories/SkeletalMeshImport.cpp b/Engine/Source/Editor/UnrealEd/Private/Factories/SkeletalMeshImport.cpp
--- a/Engine/Source/Editor/UnrealEd/Private/Factories/SkeletalMeshImport.cpp
+++ b/Engine/Source/Editor/UnrealEd/Private/Factories/SkeletalMeshImport.cpp
@@ -298,9 +298,26 @@ ExistingSkelMeshData* SaveExistingSkelMeshData(USkeletalMesh* ExistingSkelMesh,
 							OriginalSection.bCastShadow = ReduceSection.bCastShadow;
 							OriginalSection.bRecomputeTangent = ReduceSection.bRecomputeTangent;
 							OriginalSection.GenerateUpToLodIndex = ReduceSection.GenerateUpToLodIndex;
+							break;
 						}
 					}
 				}
+				//Set the unmatched original section data using the current UserSectionsData so we keep the user changes
+				for (int32 OriginalSectionIndex = 0; OriginalSectionIndex < OriginalLODModel.Sections.Num(); ++OriginalSectionIndex)
+				{
+					if (OriginalMatched[OriginalSectionIndex])
+					{
+						continue;
+					}
+					FSkelMeshSection& OriginalSection = OriginalLODModel.Sections[OriginalSectionIndex];
+					if (FSkelMeshSourceSectionUserData* ReduceUserSectionData = BackupLODModel->UserSectionsData.Find(OriginalSection.OriginalDataSectionIndex))
+					{
+						OriginalSection.bDisabled = ReduceUserSectionData->bDisabled;
+						OriginalSection.bCastShadow = ReduceUserSectionData->bCastShadow;
+						OriginalSection.bRecomputeTangent = ReduceUserSectionData->bRecomputeTangent;
+						OriginalSection.GenerateUpToLodIndex = ReduceUserSectionData->GenerateUpToLodIndex;
+					}
+				}
 				//Use the OriginalLODModel
 				BackupLODModel = &OriginalLODModel;
 			}
@@ -316,8 +333,8 @@ ExistingSkelMeshData* SaveExistingSkelMeshData(USkeletalMesh* ExistingSkelMesh,
 			int32 GenerateUpTo = BackupLODModel->Sections[SectionIndex].GenerateUpToLodIndex;
 			bool bDisabled = BackupLODModel->Sections[SectionIndex].bDisabled;
 			bool bBoneChunkedSection = BackupLODModel->Sections[SectionIndex].ChunkedParentSectionIndex != INDEX_NONE;
-			//Only save the original section, do not save the chunked sections
-			if (!bBoneChunkedSection && ExistingMeshDataPtr->ExistingImportMaterialOriginalNameData.IsValidIndex(SectionMaterialIndex))
+			//Save all the sections, even the chunked sections
+			if (ExistingMeshDataPtr->ExistingImportMaterialOriginalNameData.IsValidIndex(SectionMaterialIndex))
 			{
 				ExistingMeshDataPtr->ExistingImportMeshLodSectionMaterialData[LodIndex].Add(ExistingMeshLodSectionData(ExistingMeshDataPtr->ExistingImportMaterialOriginalNameData[SectionMaterialIndex], SectionCastShadow, SectionRecomputeTangents, GenerateUpTo, bDisabled));
 			}
@@ -1006,81 +1023,93 @@ void RestoreExistingSkelMeshData(ExistingSkelMeshData* MeshData, USkeletalMesh*
 			}
 		}
 		FSkeletalMeshLODModel &NewSkelMeshLodModel = SkeletalMeshImportedModel->LODModels[SafeReimportLODIndex];
-		int32 ParentSectionIndex = 0;
+		
+		const bool bIsValidSavedSectionMaterialData = MeshData->ExistingImportMeshLodSectionMaterialData.IsValidIndex(SafeReimportLODIndex) && MeshData->LastImportMeshLodSectionMaterialData.IsValidIndex(SafeReimportLODIndex);
+
+		const int32 MaxExistSectionNumber = bIsValidSavedSectionMaterialData ? FMath::Max(MeshData->ExistingImportMeshLodSectionMaterialData[SafeReimportLODIndex].Num(), MeshData->LastImportMeshLodSectionMaterialData[SafeReimportLODIndex].Num()) : 0;
+		TBitArray<> MatchedExistSectionIndex;
+		MatchedExistSectionIndex.Init(false, MaxExistSectionNumber);
 		//Restore the section changes from the old import data
 		for (int32 SectionIndex = 0; SectionIndex < NewSkelMeshLodModel.Sections.Num(); SectionIndex++)
 		{
-			int32 NewMeshSectionMaterialIndex = NewSkelMeshLodModel.Sections[SectionIndex].MaterialIndex;
-
-			if (RemapMaterial.IsValidIndex(NewMeshSectionMaterialIndex))
+			//Find the import section material index by using the RemapMaterial array. Fallback on the imported index if the remap entry is not valid
+			FSkelMeshSection& NewSection = NewSkelMeshLodModel.Sections[SectionIndex];
+			int32 RemapMaterialIndex = RemapMaterial.IsValidIndex(NewSection.MaterialIndex) ? RemapMaterial[NewSection.MaterialIndex] : NewSection.MaterialIndex;
+			if (!SkeletalMesh->Materials.IsValidIndex(RemapMaterialIndex))
 			{
-				if (SkeletalMesh->Materials.IsValidIndex(RemapMaterial[NewMeshSectionMaterialIndex]))
-				{
-					NewSkelMeshLodModel.Sections[SectionIndex].MaterialIndex = RemapMaterial[NewMeshSectionMaterialIndex];
-					if (MeshData->ExistingImportMeshLodSectionMaterialData[SafeReimportLODIndex].IsValidIndex(RemapMaterial[NewMeshSectionMaterialIndex]))
-					{
-						NewSkelMeshLodModel.Sections[SectionIndex].bCastShadow = MeshData->ExistingImportMeshLodSectionMaterialData[SafeReimportLODIndex][RemapMaterial[NewMeshSectionMaterialIndex]].bCastShadow;
-						NewSkelMeshLodModel.Sections[SectionIndex].bRecomputeTangent = MeshData->ExistingImportMeshLodSectionMaterialData[SafeReimportLODIndex][RemapMaterial[NewMeshSectionMaterialIndex]].bRecomputeTangents;
-						NewSkelMeshLodModel.Sections[SectionIndex].GenerateUpToLodIndex = MeshData->ExistingImportMeshLodSectionMaterialData[SafeReimportLODIndex][RemapMaterial[NewMeshSectionMaterialIndex]].GenerateUpTo;
-						NewSkelMeshLodModel.Sections[SectionIndex].bDisabled = MeshData->ExistingImportMeshLodSectionMaterialData[SafeReimportLODIndex][RemapMaterial[NewMeshSectionMaterialIndex]].bDisabled;
-						bool bBoneChunkedSection = NewSkelMeshLodModel.Sections[SectionIndex].ChunkedParentSectionIndex >= 0;
-						int32 ParentOriginalSectionIndex = NewSkelMeshLodModel.Sections[SectionIndex].OriginalDataSectionIndex;
-						if (!bBoneChunkedSection)
-						{
-							//Set the new Parent Index
-							ParentSectionIndex = SectionIndex;
-							FSkelMeshSourceSectionUserData& UserSectionData = NewSkelMeshLodModel.UserSectionsData.FindOrAdd(ParentOriginalSectionIndex);
-							UserSectionData.bDisabled = NewSkelMeshLodModel.Sections[SectionIndex].bDisabled;
-							UserSectionData.bCastShadow = NewSkelMeshLodModel.Sections[SectionIndex].bCastShadow;
-							UserSectionData.bRecomputeTangent = NewSkelMeshLodModel.Sections[SectionIndex].bRecomputeTangent;
-							UserSectionData.GenerateUpToLodIndex = NewSkelMeshLodModel.Sections[SectionIndex].GenerateUpToLodIndex;
-							//The cloth will be rebind later after the reimport is done
-						}
-					}
-				}
+				//We have an invalid material section, in this case we set the material index to 0
+				NewSection.MaterialIndex = 0;
+				UE_LOG(LogSkeletalMeshImport, Display, TEXT("Reimport material match issue: Invalid RemapMaterialIndex [%d], will make it point to material index [0]"), RemapMaterialIndex);
+				continue;
 			}
-
-			if (MeshData->LastImportMeshLodSectionMaterialData.Num() < 1 || !MeshData->LastImportMeshLodSectionMaterialData.IsValidIndex(SafeReimportLODIndex) || MeshData->LastImportMeshLodSectionMaterialData[SafeReimportLODIndex].Num() <= SectionIndex ||
-				MeshData->ExistingImportMeshLodSectionMaterialData.Num() < 1 || !MeshData->ExistingImportMeshLodSectionMaterialData.IsValidIndex(SafeReimportLODIndex) || MeshData->ExistingImportMeshLodSectionMaterialData[SafeReimportLODIndex].Num() <= SectionIndex)
+			NewSection.MaterialIndex = RemapMaterialIndex;
+			
+			//skip the rest of the loop if we do not have valid saved data
+			if (!bIsValidSavedSectionMaterialData)
 			{
-				break;
+				continue;
 			}
-
-			// If the new mesh has different number of sections than last import, then we shouldn't remap section materials.
-			if (NewSkelMeshLodModel.Sections.Num() == MeshData->LastImportMeshLodSectionMaterialData[SafeReimportLODIndex].Num())
+			//Get the RemapMaterial section Imported material slot name. We need it to match the saved existing section, so we can put back the saved existing section data
+			FName CurrentSectionImportedMaterialName = SkeletalMesh->Materials[RemapMaterialIndex].ImportedMaterialSlotName;
+			for (int32 ExistSectionIndex = 0; ExistSectionIndex < MaxExistSectionNumber; ++ExistSectionIndex)
 			{
-				int32 CurrentSectionMaterialIndex = NewSkelMeshLodModel.Sections[SectionIndex].MaterialIndex;
-				FName CurrentSectionImportedMaterialName = SkeletalMesh->Materials[CurrentSectionMaterialIndex].ImportedMaterialSlotName;
-				for (int32 ExistSectionIndex = 0; ExistSectionIndex < MeshData->LastImportMeshLodSectionMaterialData[SafeReimportLODIndex].Num(); ++ExistSectionIndex)
+				//Skip already matched exist section
+				if (MatchedExistSectionIndex[ExistSectionIndex])
 				{
-					if (!MeshData->LastImportMeshLodSectionMaterialData[SafeReimportLODIndex].IsValidIndex(ExistSectionIndex) || !MeshData->ExistingImportMeshLodSectionMaterialData[SafeReimportLODIndex].IsValidIndex(ExistSectionIndex))
-					{
-						continue;
-					}
-					//The last import slot name is use to match the New import slot name
-					//If the user has change the slot in the editor the ExistMeshSectionSlotName will be different.
-					//This is why we use the ExistMeshSectionSlotName to restore the data and use the LastImportedMeshSectionSlotName to rematch the section
-					FName LastImportedMeshSectionSlotName = MeshData->LastImportMeshLodSectionMaterialData[SafeReimportLODIndex][ExistSectionIndex];
-					if (LastImportedMeshSectionSlotName != CurrentSectionImportedMaterialName)
-					{
-						//This material do not match
-						continue;
-					}
-					//Restore the material slot
-					FName ExistMeshSectionSlotName = MeshData->ExistingImportMeshLodSectionMaterialData[SafeReimportLODIndex][ExistSectionIndex].ImportedMaterialSlotName;
+					continue;
+				}
+				//Verify we have valid existing section data, if not break from the loop higher index wont be valid
+				if (!MeshData->LastImportMeshLodSectionMaterialData[SafeReimportLODIndex].IsValidIndex(ExistSectionIndex) || !MeshData->ExistingImportMeshLodSectionMaterialData[SafeReimportLODIndex].IsValidIndex(ExistSectionIndex))
+				{
+					break;
+				}
 
-					//Override the new section material index to use the one that the user set
-					for (int32 RemapMaterialIndex = 0; RemapMaterialIndex < SkeletalMesh->Materials.Num(); ++RemapMaterialIndex)
+				//Get the Last imported skelmesh section slot import name
+				FName OriginalImportMeshSectionSlotName = MeshData->LastImportMeshLodSectionMaterialData[SafeReimportLODIndex][ExistSectionIndex];
+				if (OriginalImportMeshSectionSlotName != CurrentSectionImportedMaterialName)
+				{
+					//Skip until we found a match between the last import
+					continue;
+				}
+
+				//We have a match put back the data
+				NewSection.bCastShadow = MeshData->ExistingImportMeshLodSectionMaterialData[SafeReimportLODIndex][ExistSectionIndex].bCastShadow;
+				NewSection.bRecomputeTangent = MeshData->ExistingImportMeshLodSectionMaterialData[SafeReimportLODIndex][ExistSectionIndex].bRecomputeTangents;
+				NewSection.GenerateUpToLodIndex = MeshData->ExistingImportMeshLodSectionMaterialData[SafeReimportLODIndex][ExistSectionIndex].GenerateUpTo;
+				NewSection.bDisabled = MeshData->ExistingImportMeshLodSectionMaterialData[SafeReimportLODIndex][ExistSectionIndex].bDisabled;
+				bool bBoneChunkedSection = NewSection.ChunkedParentSectionIndex >= 0;
+				int32 ParentOriginalSectionIndex = NewSection.OriginalDataSectionIndex;
+				if (!bBoneChunkedSection)
+				{
+					//Set the new Parent Index
+					FSkelMeshSourceSectionUserData& UserSectionData = NewSkelMeshLodModel.UserSectionsData.FindOrAdd(ParentOriginalSectionIndex);
+					UserSectionData.bDisabled = NewSection.bDisabled;
+					UserSectionData.bCastShadow = NewSection.bCastShadow;
+					UserSectionData.bRecomputeTangent = NewSection.bRecomputeTangent;
+					UserSectionData.GenerateUpToLodIndex = NewSection.GenerateUpToLodIndex;
+					//The cloth will be rebind later after the reimport is done
+				}
+				//Set the matched section to true to avoid using it again
+				MatchedExistSectionIndex[ExistSectionIndex] = true;
+
+				//find the corresponding current slot name in the skeletal mesh materials list to remap properly the material index, in case the user have change it before re-importing
+				FName ExistMeshSectionSlotName = MeshData->ExistingImportMeshLodSectionMaterialData[SafeReimportLODIndex][ExistSectionIndex].ImportedMaterialSlotName;
+				{
+					for (int32 SkelMeshMaterialIndex = 0; SkelMeshMaterialIndex < SkeletalMesh->Materials.Num(); ++SkelMeshMaterialIndex)
 					{
-						const FSkeletalMaterial &NewSectionMaterial = SkeletalMesh->Materials[RemapMaterialIndex];
+						const FSkeletalMaterial &NewSectionMaterial = SkeletalMesh->Materials[SkelMeshMaterialIndex];
 						if (NewSectionMaterial.ImportedMaterialSlotName == ExistMeshSectionSlotName)
 						{
-							NewSkelMeshLodModel.Sections[SectionIndex].MaterialIndex = RemapMaterialIndex;
+							if (ExistMeshSectionSlotName != OriginalImportMeshSectionSlotName)
+							{
+								NewSection.MaterialIndex = SkelMeshMaterialIndex;
+							}
 							break;
 						}
 					}
-					break;
 				}
+				//Break because we found a match and have restore the data for this SectionIndex
+				break;
 			}
 		}
 		//Make sure we reset the User section array to only what we have in the fbx
diff --git a/Engine/Source/Runtime/Engine/Classes/Engine/SkeletalMesh.h b/Engine/Source/Runtime/Engine/Classes/Engine/SkeletalMesh.h
--- a/Engine/Source/Runtime/Engine/Classes/Engine/SkeletalMesh.h
+++ b/Engine/Source/Runtime/Engine/Classes/Engine/SkeletalMesh.h
@@ -1044,6 +1044,11 @@ public:
 #if WITH_EDITOR
 private:
 	int32 PostEditChangeStackCounter;
+
+	//When loading a legacy asset (saved before the skeletalmesh build refactor), we need to create the user sections data.
+	//This function should be call only in the PostLoad
+	void CreateUserSectionsDataForLegacyAssets();
+
 public:
 	//We want to avoid calling post edit change multiple time during import and build process.
 
@@ -1076,6 +1081,7 @@ public:
 	virtual void GetAssetRegistryTagMetadata(TMap<FName, FAssetRegistryTagMetadata>& OutMetadata) const override;
 
 	void UpdateGenerateUpToData();
+
 #endif // WITH_EDITOR
 	virtual void BeginDestroy() override;
 	virtual bool IsReadyForFinishDestroy() override;
diff --git a/Engine/Source/Runtime/Engine/Private/Rendering/SkeletalMeshLODImporterData.cpp b/Engine/Source/Runtime/Engine/Private/Rendering/SkeletalMeshLODImporterData.cpp
--- a/Engine/Source/Runtime/Engine/Private/Rendering/SkeletalMeshLODImporterData.cpp
+++ b/Engine/Source/Runtime/Engine/Private/Rendering/SkeletalMeshLODImporterData.cpp
@@ -531,6 +531,12 @@ void FReductionBaseSkeletalMeshBulkData::LoadReductionData(FSkeletalMeshLODModel
 			Ar << ReductionSkeletalMeshData;
 
 			CacheGeometryInfo(BaseLODModel);
+
+			//This call will filled missing chunked data for old asset that cannot build normal and chunking (not re import since the new skeletal mesh chunk build refactor)
+			if (!BaseLODModel.bIsBuildDataAvailable)
+			{
+				BaseLODModel.UpdateChunkedSectionInfo(Owner ? Owner->GetName() : FString(TEXT("")));
+			}
 		}
 		// Unlock the bulk data
 	}
diff --git a/Engine/Source/Runtime/Engine/Private/SkeletalMesh.cpp b/Engine/Source/Runtime/Engine/Private/SkeletalMesh.cpp
--- a/Engine/Source/Runtime/Engine/Private/SkeletalMesh.cpp
+++ b/Engine/Source/Runtime/Engine/Private/SkeletalMesh.cpp
@@ -2214,6 +2214,103 @@ void USkeletalMesh::ForceBulkDataResident(const int32 LODIndex)
 {
 	GetMeshEditorData().GetLODImportedData(LODIndex).GetBulkData().ForceBulkDataResident();
 }
+
+void USkeletalMesh::CreateUserSectionsDataForLegacyAssets()
+{
+	//We want to avoid changing the ddc if we load an old asset.
+	//This bool should be put to false at the end of the postload, if there is another posteditchange call after a new ddc will be created
+	UseLegacyMeshDerivedDataKey = true;
+	//Fill up the Section ChunkedParentSectionIndex and OriginalDataSectionIndex
+	//We also want to create the UserSectionsData structure so the user can change the section data
+	for (int32 LodIndex = 0; LodIndex < LODInfo.Num(); LodIndex++)
+	{
+		FSkeletalMeshLODModel& ThisLODModel = ImportedModel->LODModels[LodIndex];
+		FSkeletalMeshLODInfo* ThisLODInfo = GetLODInfo(LodIndex);
+		check(ThisLODInfo);
+
+		//Reset the reduction setting to a non active state if the asset has active reduction but have no RawSkeletalMeshBulkData (we cannot reduce it)
+		const bool bIsLODReductionActive = IsReductionActive(LodIndex);
+
+
+		bool bMustUseReductionSourceData = bIsLODReductionActive
+			&& ThisLODInfo->bHasBeenSimplified
+			&& ImportedModel->OriginalReductionSourceMeshData.IsValidIndex(LodIndex)
+			&& !(ImportedModel->OriginalReductionSourceMeshData[LodIndex]->IsEmpty());
+
+		if (bIsLODReductionActive && !ThisLODInfo->bHasBeenSimplified && IsLODImportedDataEmpty(LodIndex))
+		{
+			if (LodIndex > ThisLODInfo->ReductionSettings.BaseLOD)
+			{
+				ThisLODInfo->bHasBeenSimplified = true;
+			}
+			else if (LodIndex == ThisLODInfo->ReductionSettings.BaseLOD)
+			{
+				if (ThisLODInfo->ReductionSettings.TerminationCriterion == SkeletalMeshTerminationCriterion::SMTC_AbsNumOfTriangles
+					|| ThisLODInfo->ReductionSettings.TerminationCriterion == SkeletalMeshTerminationCriterion::SMTC_AbsNumOfVerts
+					|| ThisLODInfo->ReductionSettings.TerminationCriterion == SkeletalMeshTerminationCriterion::SMTC_AbsTriangleOrVert)
+				{
+					//MaxNum.... cannot be inactive, switch to NumOfTriangle
+					ThisLODInfo->ReductionSettings.TerminationCriterion = SMTC_NumOfTriangles;
+				}
+
+				//Now that we use triangle or vert num, set an inactive value
+				if (ThisLODInfo->ReductionSettings.TerminationCriterion == SkeletalMeshTerminationCriterion::SMTC_NumOfTriangles
+					|| ThisLODInfo->ReductionSettings.TerminationCriterion == SkeletalMeshTerminationCriterion::SMTC_TriangleOrVert)
+				{
+					ThisLODInfo->ReductionSettings.NumOfTrianglesPercentage = 1.0f;
+				}
+				if (ThisLODInfo->ReductionSettings.TerminationCriterion == SkeletalMeshTerminationCriterion::SMTC_NumOfVerts
+					|| ThisLODInfo->ReductionSettings.TerminationCriterion == SkeletalMeshTerminationCriterion::SMTC_TriangleOrVert)
+				{
+					ThisLODInfo->ReductionSettings.NumOfVertPercentage = 1.0f;
+				}
+			}
+			bMustUseReductionSourceData = false;
+		}
+		ThisLODModel.UpdateChunkedSectionInfo(GetName());
+
+		if (bMustUseReductionSourceData)
+		{
+			//We must load the reduction source model, since reduction can remove section
+			FSkeletalMeshLODModel ReductionSrcLODModel;
+			TMap<FString, TArray<FMorphTargetDelta>> TmpMorphTargetData;
+			ImportedModel->OriginalReductionSourceMeshData[LodIndex]->LoadReductionData(ReductionSrcLODModel, TmpMorphTargetData, this);
+			
+			//Fill the user data with the original value
+			TMap<int32, FSkelMeshSourceSectionUserData> BackupUserSectionsData = ThisLODModel.UserSectionsData;
+			ThisLODModel.UserSectionsData.Empty();
+
+			ThisLODModel.UserSectionsData = ReductionSrcLODModel.UserSectionsData;
+
+			//Now restore the reduce section user change and adjust the originalDataSectionIndex to point on the correct UserSectionData
+			TBitArray<> SourceSectionMatched;
+			SourceSectionMatched.Init(false, ReductionSrcLODModel.Sections.Num());
+			for (int32 SectionIndex = 0; SectionIndex < ThisLODModel.Sections.Num(); ++SectionIndex)
+			{
+				FSkelMeshSection& Section = ThisLODModel.Sections[SectionIndex];
+				FSkelMeshSourceSectionUserData& BackupUserData = FSkelMeshSourceSectionUserData::GetSourceSectionUserData(BackupUserSectionsData, Section);
+				for (int32 SourceSectionIndex = 0; SourceSectionIndex < ReductionSrcLODModel.Sections.Num(); ++SourceSectionIndex)
+				{
+					if (SourceSectionMatched[SourceSectionIndex])
+					{
+						continue;
+					}
+					FSkelMeshSection& SourceSection = ReductionSrcLODModel.Sections[SourceSectionIndex];
+					FSkelMeshSourceSectionUserData& UserData = FSkelMeshSourceSectionUserData::GetSourceSectionUserData(ThisLODModel.UserSectionsData, SourceSection);
+					if (Section.MaterialIndex == SourceSection.MaterialIndex)
+					{
+						Section.OriginalDataSectionIndex = SourceSection.OriginalDataSectionIndex;
+						UserData = BackupUserData;
+						SourceSectionMatched[SourceSectionIndex] = true;
+						break;
+					}
+				}
+			}
+			ThisLODModel.SyncronizeUserSectionsDataArray();
+		}
+	}
+}
+
 #endif // WITH_EDITOR
 
 void USkeletalMesh::PostLoad()
@@ -2364,51 +2461,7 @@ void USkeletalMesh::PostLoad()
 
 		if (GetLinkerCustomVersion(FEditorObjectVersion::GUID) < FEditorObjectVersion::SkeletalMeshBuildRefactor)
 		{
-			//We want to avoid changing the ddc if we load an old asset.
-			//This bool should be put to false at the end of the postload, if there is another posteditchange call after a new ddc will be created
-			UseLegacyMeshDerivedDataKey = true;
-			//Fill up the Section ChunkedParentSectionIndex and OriginalDataSectionIndex
-			//We also want to create the UserSectionsData structure so the user can change the section data
-			for (int32 LodIndex = 0; LodIndex < LODInfo.Num(); LodIndex++)
-			{
-				FSkeletalMeshLODModel& ThisLODModel = ImportedModel->LODModels[LodIndex];
-				FSkeletalMeshLODInfo* ThisLODInfo = GetLODInfo(LodIndex);
-				check(ThisLODInfo);
-				//Reset the reduction setting to a non active state if the asset has active reduction but have no RawSkeletalMeshBulkData (we cannot reduce it)
-				if (IsReductionActive(LodIndex) && !ThisLODInfo->bHasBeenSimplified && IsLODImportedDataEmpty(LodIndex))
-				{
-					if (LodIndex > ThisLODInfo->ReductionSettings.BaseLOD)
-					{
-						ThisLODInfo->bHasBeenSimplified = true;
-					}
-					else if (LodIndex == ThisLODInfo->ReductionSettings.BaseLOD)
-					{
-						if (ThisLODInfo->ReductionSettings.TerminationCriterion == SkeletalMeshTerminationCriterion::SMTC_AbsNumOfTriangles
-							|| ThisLODInfo->ReductionSettings.TerminationCriterion == SkeletalMeshTerminationCriterion::SMTC_AbsNumOfVerts
-							|| ThisLODInfo->ReductionSettings.TerminationCriterion == SkeletalMeshTerminationCriterion::SMTC_AbsTriangleOrVert)
-						{
-							//MaxNum.... cannot be inactive, switch to NumOfTriangle
-							ThisLODInfo->ReductionSettings.TerminationCriterion = SMTC_NumOfTriangles;
-						}
-
-						//Now that we use trinagle or vert num, set an inactive value
-						if (ThisLODInfo->ReductionSettings.TerminationCriterion == SkeletalMeshTerminationCriterion::SMTC_NumOfTriangles
-							|| ThisLODInfo->ReductionSettings.TerminationCriterion == SkeletalMeshTerminationCriterion::SMTC_TriangleOrVert)
-						{
-							ThisLODInfo->ReductionSettings.NumOfTrianglesPercentage = 1.0f;
-						}
-						if (ThisLODInfo->ReductionSettings.TerminationCriterion == SkeletalMeshTerminationCriterion::SMTC_NumOfVerts
-							|| ThisLODInfo->ReductionSettings.TerminationCriterion == SkeletalMeshTerminationCriterion::SMTC_TriangleOrVert)
-						{
-							ThisLODInfo->ReductionSettings.NumOfVertPercentage = 1.0f;
-						}
-					}
-				}
-				if (!IsLODImportedDataBuildAvailable(LodIndex))
-				{
-					ThisLODModel.UpdateChunkedSectionInfo(GetName(), ThisLODInfo->LODMaterialMap);
-				}
-			}
+			CreateUserSectionsDataForLegacyAssets();
 		}
 
 		if (GetResourceForRendering() == nullptr)
diff --git a/Engine/Source/Runtime/Engine/Private/SkeletalMeshLODModel.cpp b/Engine/Source/Runtime/Engine/Private/SkeletalMeshLODModel.cpp
--- a/Engine/Source/Runtime/Engine/Private/SkeletalMeshLODModel.cpp
+++ b/Engine/Source/Runtime/Engine/Private/SkeletalMeshLODModel.cpp
@@ -1103,25 +1103,7 @@ FString FSkeletalMeshLODModel::GetLODModelDeriveDataKey() const
 	return KeySuffix;
 }
 
-bool FSkeletalMeshLODModel::CopyStructure(FSkeletalMeshLODModel* Destination, FSkeletalMeshLODModel* Source)
-{
-	if (Source->RawPointIndices.IsLocked() || Source->LegacyRawPointIndices.IsLocked() || Source->RawSkeletalMeshBulkData_DEPRECATED.GetBulkData().IsLocked())
-	{
-		return false;
-	}
-	// Bulk data arrays need to be locked before a copy can be made.
-	Source->RawPointIndices.Lock(LOCK_READ_ONLY);
-	Source->LegacyRawPointIndices.Lock(LOCK_READ_ONLY);
-	Source->RawSkeletalMeshBulkData_DEPRECATED.GetBulkData().Lock(LOCK_READ_ONLY);
-	*Destination = *Source;
-	Source->RawSkeletalMeshBulkData_DEPRECATED.GetBulkData().Unlock();
-	Source->RawPointIndices.Unlock();
-	Source->LegacyRawPointIndices.Unlock();
-
-	return true;
-}
-
-void FSkeletalMeshLODModel::UpdateChunkedSectionInfo(const FString& SkeletalMeshName, TArray<int32>& LODMaterialMap)
+void FSkeletalMeshLODModel::UpdateChunkedSectionInfo(const FString& SkeletalMeshName)
 {
 	int32 LODModelSectionNum = Sections.Num();
 	//Fill the ChunkedParentSectionIndex data, we assume that every section using the same material are chunked
@@ -1176,4 +1158,22 @@ void FSkeletalMeshLODModel::UpdateChunkedSectionInfo(const FString& SkeletalMesh
 	}
 }
 
+bool FSkeletalMeshLODModel::CopyStructure(FSkeletalMeshLODModel* Destination, FSkeletalMeshLODModel* Source)
+{
+	if (Source->RawPointIndices.IsLocked() || Source->LegacyRawPointIndices.IsLocked() || Source->RawSkeletalMeshBulkData_DEPRECATED.GetBulkData().IsLocked())
+	{
+		return false;
+	}
+	// Bulk data arrays need to be locked before a copy can be made.
+	Source->RawPointIndices.Lock(LOCK_READ_ONLY);
+	Source->LegacyRawPointIndices.Lock(LOCK_READ_ONLY);
+	Source->RawSkeletalMeshBulkData_DEPRECATED.GetBulkData().Lock(LOCK_READ_ONLY);
+	*Destination = *Source;
+	Source->RawSkeletalMeshBulkData_DEPRECATED.GetBulkData().Unlock();
+	Source->RawPointIndices.Unlock();
+	Source->LegacyRawPointIndices.Unlock();
+
+	return true;
+}
+
 #endif // WITH_EDITOR
diff --git a/Engine/Source/Runtime/Engine/Public/Rendering/SkeletalMeshLODModel.h b/Engine/Source/Runtime/Engine/Public/Rendering/SkeletalMeshLODModel.h
--- a/Engine/Source/Runtime/Engine/Public/Rendering/SkeletalMeshLODModel.h
+++ b/Engine/Source/Runtime/Engine/Public/Rendering/SkeletalMeshLODModel.h
@@ -273,6 +273,24 @@ struct FSkelMeshSourceSectionUserData
 		return (ClothingData.AssetGuid.IsValid());
 	}
 
+	static FSkelMeshSourceSectionUserData& GetSourceSectionUserData(TMap<int32, FSkelMeshSourceSectionUserData>& UserSectionsData, const FSkelMeshSection& Section)
+	{
+		FSkelMeshSourceSectionUserData* UserSectionData = UserSectionsData.Find(Section.OriginalDataSectionIndex);
+		if (!UserSectionData)
+		{
+			//If the UserSectionData do not exist add it and copy from the section data
+			UserSectionData = &UserSectionsData.Add(Section.OriginalDataSectionIndex);
+			UserSectionData->bCastShadow = Section.bCastShadow;
+			UserSectionData->bDisabled = Section.bDisabled;
+			UserSectionData->bRecomputeTangent = Section.bRecomputeTangent;
+			UserSectionData->GenerateUpToLodIndex = Section.GenerateUpToLodIndex;
+			UserSectionData->CorrespondClothAssetIndex = Section.CorrespondClothAssetIndex;
+			UserSectionData->ClothingData.AssetGuid = Section.ClothingData.AssetGuid;
+			UserSectionData->ClothingData.AssetLodIndex = Section.ClothingData.AssetLodIndex;
+		}
+		check(UserSectionData);
+		return *UserSectionData;
+	}
 	// Serialization.
 	friend FArchive& operator<<(FArchive& Ar, FSkelMeshSourceSectionUserData& S);
 };
@@ -420,7 +438,7 @@ public:
 	* This function will update the chunked information for each section. Only old data before the 
 	* skeletal mesh build refactor should need to call this function.
 	*/
-	ENGINE_API void UpdateChunkedSectionInfo(const FString& SkeletalMeshName, TArray<int32>& LODMaterialMap);
+	ENGINE_API void UpdateChunkedSectionInfo(const FString& SkeletalMeshName);
 };
 
 #endif // WITH_EDITOR
\ No newline at end of file
