Make sure a 1.0 vert/tri reduction will reduce if the srcModel have more vert/tri than the MaxNumOfVertsPercentage/MaxNumOfTrianglesPercentage value.

diff --git a/Engine/Plugins/Experimental/SkeletalReduction/Source/Private/SkeletalMeshReductionPlugin.cpp b/Engine/Plugins/Experimental/SkeletalReduction/Source/Private/SkeletalMeshReductionPlugin.cpp
--- a/Engine/Plugins/Experimental/SkeletalReduction/Source/Private/SkeletalMeshReductionPlugin.cpp
+++ b/Engine/Plugins/Experimental/SkeletalReduction/Source/Private/SkeletalMeshReductionPlugin.cpp
@@ -1645,10 +1645,21 @@ bool FQuadricSkeletalMeshReduction::ReduceSkeletalLODModel( const FSkeletalMeshL
 {
 
 	const int32 SrcNumVerts = SrcModel.NumVertices;
-	
+	auto GetLODModelTriangleCount = [](const FSkeletalMeshLODModel& LODModel)->uint32
+	{
+		uint32 NumTriangles = 0;
+		for (const FSkelMeshSection& Section : LODModel.Sections)
+		{
+			
+			NumTriangles += Section.NumTriangles >= 0 ? (uint32)Section.NumTriangles : 0;
+		}
+		return NumTriangles;
+	};
+	const uint32 SrcNumTriangles = GetLODModelTriangleCount(SrcModel);
+
 	// Parameters for Simplification etc
-	const bool bUseVertexPercentCriterion   = ((Settings.TerminationCriterion == SMTC_NumOfVerts     || Settings.TerminationCriterion == SMTC_TriangleOrVert) && Settings.NumOfVertPercentage < 1.f) ;
-	const bool bUseTrianglePercentCriterion = ((Settings.TerminationCriterion == SMTC_NumOfTriangles || Settings.TerminationCriterion == SMTC_TriangleOrVert) && Settings.NumOfTrianglesPercentage < 1.f);
+	const bool bUseVertexPercentCriterion   = ((Settings.TerminationCriterion == SMTC_NumOfVerts     || Settings.TerminationCriterion == SMTC_TriangleOrVert) && (Settings.NumOfVertPercentage < 1.f || Settings.MaxNumOfVertsPercentage < (uint32)SrcNumVerts)) ;
+	const bool bUseTrianglePercentCriterion = ((Settings.TerminationCriterion == SMTC_NumOfTriangles || Settings.TerminationCriterion == SMTC_TriangleOrVert) && (Settings.NumOfTrianglesPercentage < 1.f || Settings.MaxNumOfTrianglesPercentage < SrcNumTriangles));
 
 	const bool bUseMaxVertexCriterion   = ((Settings.TerminationCriterion == SMTC_AbsNumOfVerts || Settings.TerminationCriterion == SMTC_AbsTriangleOrVert) && SrcNumVerts);
 	const bool bUseMaxTriangleCriterion = ((Settings.TerminationCriterion == SMTC_AbsNumOfTriangles || Settings.TerminationCriterion == SMTC_AbsTriangleOrVert) && Settings.MaxNumOfTriangles < INT32_MAX);
