Fix unsigned/sign mismatch error in the skeletalmesh quadric reduction plugin.

diff --git a/Engine/Plugins/Experimental/SkeletalReduction/Source/Private/SkeletalMeshReductionPlugin.cpp b/Engine/Plugins/Experimental/SkeletalReduction/Source/Private/SkeletalMeshReductionPlugin.cpp
--- a/Engine/Plugins/Experimental/SkeletalReduction/Source/Private/SkeletalMeshReductionPlugin.cpp
+++ b/Engine/Plugins/Experimental/SkeletalReduction/Source/Private/SkeletalMeshReductionPlugin.cpp
@@ -928,26 +928,27 @@ float FQuadricSkeletalMeshReduction::SimplifyMesh( const FSkeletalMeshOptimizati
 	// We can support a stopping criteria based on the MaxDistance the new vertex is from the plans of the source triangles.
 	// but there seems to be no good use for this.  We are better off just using triangle count.
 	const float MaxDist = FLT_MAX; // (Settings.ReductionMethod != SkeletalMeshOptimizationType::SMOT_NumOfTriangles) ? Settings.MaxDeviationPercentage * Bounds.SphereRadius : FLT_MAX;
-	const int32 SrcTriNum = Mesh.NumIndices() / 3;
+	const uint32 SrcTriNum = Mesh.NumIndices() / 3;
 	const float TriangleRetainRatio = FMath::Clamp(Settings.NumOfTrianglesPercentage, 0.f, 1.f);
-	int32 TargetTriNum = (bUseTrianglePercentCriterion) ? FMath::CeilToInt(TriangleRetainRatio * SrcTriNum) : Settings.MaxNumOfTriangles;
+	uint32 TargetTriNum = (bUseTrianglePercentCriterion) ? FMath::CeilToInt(TriangleRetainRatio * SrcTriNum) : Settings.MaxNumOfTriangles;
 	if (Settings.TerminationCriterion == SkeletalMeshTerminationCriterion::SMTC_NumOfTriangles)
 	{
-		TargetTriNum = FMath::Min(TargetTriNum, int32(Settings.MaxNumOfTrianglesPercentage));
+		TargetTriNum = FMath::Min(TargetTriNum, Settings.MaxNumOfTrianglesPercentage);
 	}
 
-	const int32 MinTriNumToRetain = (bUseTrianglePercentCriterion || bUseMaxTrisNumCriterion) ? FMath::Max(4, TargetTriNum) : 4;
+	const uint32 MinTriNumToRetain = (bUseTrianglePercentCriterion || bUseMaxTrisNumCriterion) ? FMath::Max((uint32)4, TargetTriNum) : 4;
 	const float MaxCollapseCost = FLT_MAX;
 
-	const int32 SrcVertNum = Mesh.NumVertices();
+	const uint32 SrcVertNum = Mesh.NumVertices();
 	const float VertRetainRatio = FMath::Clamp(Settings.NumOfVertPercentage, 0.f, 1.f);
-	int32 TargetVertNum = (bUseVertexPercentCriterion) ? FMath::CeilToInt(VertRetainRatio * SrcVertNum) : Settings.MaxNumOfVerts + 1;
+	uint32 TargetVertNum = (bUseVertexPercentCriterion) ? FMath::CeilToInt(VertRetainRatio * SrcVertNum) : Settings.MaxNumOfVerts + 1;
 	if (Settings.TerminationCriterion == SkeletalMeshTerminationCriterion::SMTC_NumOfVerts)
 	{
-		TargetVertNum = FMath::Min(TargetVertNum, int32(Settings.MaxNumOfVertsPercentage + 1));
+		uint32 Delta = Settings.MaxNumOfVertsPercentage < MAX_uint32 ? 1 : 0;
+		TargetVertNum = FMath::Min(TargetVertNum, Settings.MaxNumOfVertsPercentage + Delta);
 	}
 
-	const int32 MinVerNumToRetain = (bUseVertexPercentCriterion || bUseMaxVertNumCriterion) ? FMath::Max(6, TargetVertNum) : 6;
+	const uint32 MinVerNumToRetain = (bUseVertexPercentCriterion || bUseMaxVertNumCriterion) ? FMath::Max((uint32)6, TargetVertNum) : 6;
 
 	const float VolumeImportance      = FMath::Clamp(Settings.VolumeImportance, 0.f, 2.f);
 	const bool bLockEdges             = Settings.bLockEdges;
@@ -1644,7 +1645,7 @@ bool FQuadricSkeletalMeshReduction::ReduceSkeletalLODModel( const FSkeletalMeshL
                                                            ) const
 {
 
-	const int32 SrcNumVerts = SrcModel.NumVertices;
+	const uint32 SrcNumVerts = SrcModel.NumVertices;
 	auto GetLODModelTriangleCount = [](const FSkeletalMeshLODModel& LODModel)->uint32
 	{
 		uint32 NumTriangles = 0;
@@ -1658,7 +1659,7 @@ bool FQuadricSkeletalMeshReduction::ReduceSkeletalLODModel( const FSkeletalMeshL
 	const uint32 SrcNumTriangles = GetLODModelTriangleCount(SrcModel);
 
 	// Parameters for Simplification etc
-	const bool bUseVertexPercentCriterion   = ((Settings.TerminationCriterion == SMTC_NumOfVerts     || Settings.TerminationCriterion == SMTC_TriangleOrVert) && (Settings.NumOfVertPercentage < 1.f || Settings.MaxNumOfVertsPercentage < (uint32)SrcNumVerts)) ;
+	const bool bUseVertexPercentCriterion   = ((Settings.TerminationCriterion == SMTC_NumOfVerts     || Settings.TerminationCriterion == SMTC_TriangleOrVert) && (Settings.NumOfVertPercentage < 1.f || Settings.MaxNumOfVertsPercentage < SrcNumVerts)) ;
 	const bool bUseTrianglePercentCriterion = ((Settings.TerminationCriterion == SMTC_NumOfTriangles || Settings.TerminationCriterion == SMTC_TriangleOrVert) && (Settings.NumOfTrianglesPercentage < 1.f || Settings.MaxNumOfTrianglesPercentage < SrcNumTriangles));
 
 	const bool bUseMaxVertexCriterion   = ((Settings.TerminationCriterion == SMTC_AbsNumOfVerts || Settings.TerminationCriterion == SMTC_AbsTriangleOrVert) && SrcNumVerts);
@@ -1727,19 +1728,19 @@ bool FQuadricSkeletalMeshReduction::ReduceSkeletalLODModel( const FSkeletalMeshL
 		if (bUseMaxVertexCriterion && OutSkeletalMeshLODModel.NumVertices > OriginalMaxNumVertsSetting && OutSkeletalMeshLODModel.NumVertices > 6)
 		{
 			const bool bTerminatedOnVertCount = (Settings.TerminationCriterion == SMTC_AbsNumOfVerts) ||
-				                                (Settings.TerminationCriterion == SMTC_AbsTriangleOrVert && !(SkinnedSkeletalMesh.NumIndices() / 3 <= (int32)Settings.MaxNumOfTriangles));
+				                                (Settings.TerminationCriterion == SMTC_AbsTriangleOrVert && !((uint32)SkinnedSkeletalMesh.NumIndices() / (uint32)3 <= Settings.MaxNumOfTriangles));
 			  
 			if (bTerminatedOnVertCount)
 			{
 				// Some verts were created by chunking - we need simplify more.
-				int32 ExcessVerts = (int32)(OutSkeletalMeshLODModel.NumVertices - OriginalMaxNumVertsSetting);
+				uint32 ExcessVerts = OutSkeletalMeshLODModel.NumVertices - OriginalMaxNumVertsSetting;
 				if (bUseVertexPercentCriterion)
 				{
-					Settings.MaxNumOfVertsPercentage = FMath::Max((int32)Settings.MaxNumOfVertsPercentage - ExcessVerts, 6);
+					Settings.MaxNumOfVertsPercentage = FMath::Max(Settings.MaxNumOfVertsPercentage - ExcessVerts, (uint32)6);
 				}
 				else
 				{
-					Settings.MaxNumOfVerts = FMath::Max((int32)Settings.MaxNumOfVerts - ExcessVerts, 6);
+					Settings.MaxNumOfVerts = FMath::Max(Settings.MaxNumOfVerts - ExcessVerts, (uint32)6);
 				}
 
 				UE_LOG(LogSkeletalMeshReduction, Log, TEXT("Chunking to limit unique bones per section generated additional vertices - continuing simplification of LOD %d "), LODIndex);
diff --git a/Engine/Plugins/Experimental/SkeletalReduction/Source/Private/SkeletalSimplifier.h b/Engine/Plugins/Experimental/SkeletalReduction/Source/Private/SkeletalSimplifier.h
--- a/Engine/Plugins/Experimental/SkeletalReduction/Source/Private/SkeletalSimplifier.h
+++ b/Engine/Plugins/Experimental/SkeletalReduction/Source/Private/SkeletalSimplifier.h
@@ -49,7 +49,7 @@ namespace SkeletalSimplifier
 	class FSimplifierTerminator : public FSimplifierTerminatorBase
 	{
 	public:
-		FSimplifierTerminator(int32 MinTri, int32 MaxTri, int32 MinVert, int32 MaxVert, float MaxCost, float MaxDist)
+		FSimplifierTerminator(uint32 MinTri, uint32 MaxTri, uint32 MinVert, uint32 MaxVert, float MaxCost, float MaxDist)
 			: FSimplifierTerminatorBase(MinTri, MinVert, MaxCost)
 			, MaxTriNumToRetain(MaxTri)
 			, MaxVertNumToRetain(MaxVert)
@@ -57,7 +57,7 @@ namespace SkeletalSimplifier
 		{}
 
 		// return true if the simplifier should terminate.
-		inline bool operator()(const int32 TriNum, const int32 VertNum, const float SqrError)
+		inline bool operator()(const uint32 TriNum, const uint32 VertNum, const float SqrError)
 		{
 			if (FSimplifierTerminatorBase::operator()(TriNum, VertNum, SqrError) && TriNum < MaxTriNumToRetain && VertNum < MaxVertNumToRetain)
 			{
@@ -70,8 +70,8 @@ namespace SkeletalSimplifier
 		}
 
 
-		int32 MaxTriNumToRetain;
-		int32 MaxVertNumToRetain;
+		uint32 MaxTriNumToRetain;
+		uint32 MaxVertNumToRetain;
 		float MaxDistance;
 
 	};
diff --git a/Engine/Source/Runtime/Engine/Private/SkeletalMeshRenderData.cpp b/Engine/Source/Runtime/Engine/Private/SkeletalMeshRenderData.cpp
--- a/Engine/Source/Runtime/Engine/Private/SkeletalMeshRenderData.cpp
+++ b/Engine/Source/Runtime/Engine/Private/SkeletalMeshRenderData.cpp
@@ -59,7 +59,7 @@ static void SerializeLODInfoForDDC(USkeletalMesh* SkeletalMesh, FString& KeySuff
 // differences, etc.) replace the version GUID below with a new one.
 // In case of merge conflicts with DDC versions, you MUST generate a new GUID
 // and set this new GUID as the version.
-#define SKELETALMESH_DERIVEDDATA_VER TEXT("BAD8180EFB1543D79AA1D10F3F5945A2")
+#define SKELETALMESH_DERIVEDDATA_VER TEXT("D63A7B68FF6E4919AE164CD429AD51D7")
 
 static const FString& GetSkeletalMeshDerivedDataVersion()
 {
