Fix static analyser CIS warning

diff --git a/Engine/Plugins/Experimental/SkeletalReduction/Source/Private/SkeletalMeshReductionPlugin.cpp b/Engine/Plugins/Experimental/SkeletalReduction/Source/Private/SkeletalMeshReductionPlugin.cpp
--- a/Engine/Plugins/Experimental/SkeletalReduction/Source/Private/SkeletalMeshReductionPlugin.cpp
+++ b/Engine/Plugins/Experimental/SkeletalReduction/Source/Private/SkeletalMeshReductionPlugin.cpp
@@ -1651,8 +1651,7 @@ bool FQuadricSkeletalMeshReduction::ReduceSkeletalLODModel( const FSkeletalMeshL
 		uint32 NumTriangles = 0;
 		for (const FSkelMeshSection& Section : LODModel.Sections)
 		{
-			
-			NumTriangles += Section.NumTriangles >= 0 ? (uint32)Section.NumTriangles : 0;
+			NumTriangles += Section.NumTriangles;
 		}
 		return NumTriangles;
 	};
