Make sure the reduction can iterate when using vertice percentage with a maximum absolute vertex.
#rnx

diff --git a/Engine/Plugins/Experimental/SkeletalReduction/Source/Private/SkeletalMeshReductionPlugin.cpp b/Engine/Plugins/Experimental/SkeletalReduction/Source/Private/SkeletalMeshReductionPlugin.cpp
--- a/Engine/Plugins/Experimental/SkeletalReduction/Source/Private/SkeletalMeshReductionPlugin.cpp
+++ b/Engine/Plugins/Experimental/SkeletalReduction/Source/Private/SkeletalMeshReductionPlugin.cpp
@@ -1724,9 +1724,10 @@ bool FQuadricSkeletalMeshReduction::ReduceSkeletalLODModel( const FSkeletalMeshL
 
 		// We may need to do additional simplification if the user specified a hard number limit for verts and
 		// the internal chunking during conversion split some verts.
-		if (bUseMaxVertexCriterion && OutSkeletalMeshLODModel.NumVertices > OriginalMaxNumVertsSetting && OutSkeletalMeshLODModel.NumVertices > 6)
+		if ((bUseMaxVertexCriterion || bUseVertexPercentCriterion) && OutSkeletalMeshLODModel.NumVertices > OriginalMaxNumVertsSetting && OutSkeletalMeshLODModel.NumVertices > 6)
 		{
-			const bool bTerminatedOnVertCount = (Settings.TerminationCriterion == SMTC_AbsNumOfVerts) ||
+			const bool bTerminatedOnVertCount = ((Settings.TerminationCriterion == SMTC_NumOfVerts || Settings.TerminationCriterion == SMTC_TriangleOrVert) && (OutSkeletalMeshLODModel.NumVertices > OriginalMaxNumVertsSetting)) ||
+												(Settings.TerminationCriterion == SMTC_AbsNumOfVerts) ||
 				                                (Settings.TerminationCriterion == SMTC_AbsTriangleOrVert && !((uint32)SkinnedSkeletalMesh.NumIndices() / (uint32)3 <= Settings.MaxNumOfTriangles));
 			  
 			if (bTerminatedOnVertCount)
