Fix crash when reimporting base mesh and also changing the alternate skinning influence (during the reimport)

diff --git a/Engine/Plugins/Experimental/SkeletalReduction/Source/Private/SkeletalMeshReductionPlugin.cpp b/Engine/Plugins/Experimental/SkeletalReduction/Source/Private/SkeletalMeshReductionPlugin.cpp
--- a/Engine/Plugins/Experimental/SkeletalReduction/Source/Private/SkeletalMeshReductionPlugin.cpp
+++ b/Engine/Plugins/Experimental/SkeletalReduction/Source/Private/SkeletalMeshReductionPlugin.cpp
@@ -1370,6 +1370,11 @@ void  FQuadricSkeletalMeshReduction::AddSourceModelInfluences( const FSkeletalMe
 			const FName& ProfileName = Profile.Key;
 			const FImportedSkinWeightProfileData& SrcImportedProfileData = Profile.Value;
 			const TArray<FRawSkinWeight>& SrcBonesAndWeights = SrcImportedProfileData.SkinWeights;
+			if (SrcBonesAndWeights.Num() < 1 || SrcImportedProfileData.SourceModelInfluences.Num() < 1)
+			{
+				//Skip empty profile
+				continue;
+			}
 
 
 			// Create the SrcModelInfluences for this profile.
