Fixing PVS 7.7 Issues:

warning V614: Uninitialized variable 'NumTris' used. Consider checking the first actual argument of the 'Resize' function.

diff --git a/Engine/Plugins/Experimental/SkeletalReduction/Source/Private/SkeletalMeshReductionSkinnedMesh.h b/Engine/Plugins/Experimental/SkeletalReduction/Source/Private/SkeletalMeshReductionSkinnedMesh.h
--- a/Engine/Plugins/Experimental/SkeletalReduction/Source/Private/SkeletalMeshReductionSkinnedMesh.h
+++ b/Engine/Plugins/Experimental/SkeletalReduction/Source/Private/SkeletalMeshReductionSkinnedMesh.h
@@ -33,13 +33,17 @@ namespace SkeletalSimplifier
 			Empty();
 		}
 
-
 		/**
 		* Constructor allocates index and vertex buffer.
 		*/
-		TSkinnedSkeletalMesh(int32 NumTriangles, int32 NumVertices)
+		TSkinnedSkeletalMesh(int32 NumTriangles, int32 NumVertices) :
+			IndexBuffer(NULL),
+			VertexBuffer(NULL),
+			NumTris(0),
+			NumVerts(0),
+			TexCoordNum(0)
 		{
-			Resize(NumTris, NumVerts);
+			Resize(NumTriangles, NumVertices);
 		}
 
 		/**
