//UE5/Release-5.0-EarlyAccess - NonUnity Compile UnrealEditor Win64 - cannot convert argument 1 from 'const UAnimSequence *' to 'const UAnimSequenceBase
#fix added includes for AnimSequence.h

diff --git a/Engine/Plugins/Experimental/SkeletalReduction/Source/Private/SkeletalMeshReductionPlugin.cpp b/Engine/Plugins/Experimental/SkeletalReduction/Source/Private/SkeletalMeshReductionPlugin.cpp
--- a/Engine/Plugins/Experimental/SkeletalReduction/Source/Private/SkeletalMeshReductionPlugin.cpp
+++ b/Engine/Plugins/Experimental/SkeletalReduction/Source/Private/SkeletalMeshReductionPlugin.cpp
@@ -25,6 +25,7 @@
 #include "ClothingAsset.h"
 #include "Factories/FbxSkeletalMeshImportData.h"
 #include "LODUtilities.h"
+#include "Animation/AnimSequence.h"
 #include "Async/Async.h"
 
 #define LOCTEXT_NAMESPACE "SkeletalMeshReduction"
@@ -2046,11 +2047,11 @@ void FQuadricSkeletalMeshReduction::ReduceSkeletalMesh(USkeletalMesh& SkeletalMe
 				// Make sure the deletion is happening on the game-thread.
 				// Deleting a structure containing bulkdata can crash when the bulkdata is detached from the archive.
 				if (IsInGameThread())
-				{
-					delete Old;
-				}
-				else
-				{
+			{
+				delete Old;
+			}
+			else
+			{
 					Async(EAsyncExecution::TaskGraphMainThread, [Old]() { delete Old; });
 				}
 			}
diff --git a/Engine/Source/Developer/MeshBoneReduction/Private/MeshBoneReduction.cpp b/Engine/Source/Developer/MeshBoneReduction/Private/MeshBoneReduction.cpp
--- a/Engine/Source/Developer/MeshBoneReduction/Private/MeshBoneReduction.cpp
+++ b/Engine/Source/Developer/MeshBoneReduction/Private/MeshBoneReduction.cpp
@@ -13,6 +13,7 @@
 #include "Rendering/SkeletalMeshRenderData.h"
 #include "AnimationBlueprintLibrary.h"
 #include "Async/ParallelFor.h"
+#include "Animation/AnimSequence.h"
 
 class FMeshBoneReductionModule : public IMeshBoneReductionModule
 {
