SkeletalSimplifier::FSimplifierMeshManager::MakeEdges: Prevent negative array size reservation when NumVerts is 0.
#rnx
#[fyi] Paul.Chipchase

diff --git a/Engine/Plugins/Experimental/SkeletalReduction/Source/Private/SkeletalSimplifierMeshManager.cpp b/Engine/Plugins/Experimental/SkeletalReduction/Source/Private/SkeletalSimplifierMeshManager.cpp
--- a/Engine/Plugins/Experimental/SkeletalReduction/Source/Private/SkeletalSimplifierMeshManager.cpp
+++ b/Engine/Plugins/Experimental/SkeletalReduction/Source/Private/SkeletalSimplifierMeshManager.cpp
@@ -325,7 +325,7 @@ void SkeletalSimplifier::FSimplifierMeshManager::MakeEdges(const TArray<SimpVert
 	// Populate the TArray of edges.
 	const int32 NumVerts = Verts.Num();
 
-	int32 maxEdgeSize = FMath::Min(3 * NumTris, 3 * NumVerts - 6);
+	int32 maxEdgeSize = FMath::Min(3 * NumTris, 3 * FMath::Max(NumVerts,2) - 6);
 	Edges.Empty(maxEdgeSize);
 	for (int i = 0; i < NumVerts; i++)
 	{
