Add an option to skeletal mesh reduction settings (bMergeCoincidentVertBones) to prevent two vertex on the same position to use different bone skinning. By default the option is true since it was the previous behavior. The option in the UI is call (Merge Coincident Vertices Bones)

#rnx

diff --git a/Engine/Plugins/Experimental/SkeletalReduction/Source/Private/SkeletalMeshReductionPlugin.cpp b/Engine/Plugins/Experimental/SkeletalReduction/Source/Private/SkeletalMeshReductionPlugin.cpp
--- a/Engine/Plugins/Experimental/SkeletalReduction/Source/Private/SkeletalMeshReductionPlugin.cpp
+++ b/Engine/Plugins/Experimental/SkeletalReduction/Source/Private/SkeletalMeshReductionPlugin.cpp
@@ -970,6 +970,7 @@ float FQuadricSkeletalMeshReduction::SimplifyMesh( const FSkeletalMeshOptimizati
 	const bool bLockEdges             = Settings.bLockEdges;
 	const bool bPreserveVolume        = (VolumeImportance > 1.e-4);
 	const bool bEnforceBoneBoundaries = Settings.bEnforceBoneBoundaries;
+	const bool bMergeCoincidentVertBones = Settings.bMergeCoincidentVertBones;
 	const bool bLockColorBoundaries   = Settings.bLockColorBounaries;
 
 	// Terminator tells the simplifier when to stop
@@ -1068,7 +1069,6 @@ float FQuadricSkeletalMeshReduction::SimplifyMesh( const FSkeletalMeshOptimizati
 
 	// Additional parameters
 
-	const bool bMergeCoincidentVertBones = true;
 	const bool bWeldVertexColorAttributes = true;
 	const float EdgeWeightValue = 128.f;
 
@@ -1078,7 +1078,8 @@ float FQuadricSkeletalMeshReduction::SimplifyMesh( const FSkeletalMeshOptimizati
 	
 	SkeletalSimplifier::FMeshSimplifier  Simplifier(Mesh.VertexBuffer, (uint32)Mesh.NumVertices(),
 		                                            Mesh.IndexBuffer, (uint32)Mesh.NumIndices(), 
-		                                            CoAlignmentLimit, VolumeImportance, bPreserveVolume,  bEnforceBoneBoundaries);
+		                                            CoAlignmentLimit, VolumeImportance, bPreserveVolume,  
+													bEnforceBoneBoundaries, bMergeCoincidentVertBones);
 
 	// The simplifier made a deep copy of the mesh.  
 
@@ -1132,7 +1133,7 @@ float FQuadricSkeletalMeshReduction::SimplifyMesh( const FSkeletalMeshOptimizati
 
 	// Copy the simplified mesh back into Mesh
 
-	Simplifier.OutputMesh(Mesh.VertexBuffer, Mesh.IndexBuffer, bMergeCoincidentVertBones, bWeldVertexColorAttributes, NULL);
+	Simplifier.OutputMesh(Mesh.VertexBuffer, Mesh.IndexBuffer, bWeldVertexColorAttributes, NULL);
 
 	// There might have some unused verts at the end of the vertex buffer that were generated by the possible duplicates
 
diff --git a/Engine/Plugins/Experimental/SkeletalReduction/Source/Private/SkeletalSimplifier.cpp b/Engine/Plugins/Experimental/SkeletalReduction/Source/Private/SkeletalSimplifier.cpp
--- a/Engine/Plugins/Experimental/SkeletalReduction/Source/Private/SkeletalSimplifier.cpp
+++ b/Engine/Plugins/Experimental/SkeletalReduction/Source/Private/SkeletalSimplifier.cpp
@@ -14,13 +14,15 @@ const int32 SkeletalSimplifier::LinearAlgebra::SymmetricMatrix::Mapping[9] = { 0
 //=============
 SkeletalSimplifier::FMeshSimplifier::FMeshSimplifier(const MeshVertType* InSrcVerts, const uint32 InNumSrcVerts,
 	const uint32* InSrcIndexes, const uint32 InNumSrcIndexes,
-	const float CoAlignmentLimit, const float VolumeImportanceValue, const bool VolumeConservation, const bool bEnforceBoundaries)
+	const float CoAlignmentLimit, const float VolumeImportanceValue, const bool VolumeConservation, const bool bEnforceBoundaries, 
+	const bool bMergeCoincidentVertBones)
 	:
 	coAlignmentLimit(CoAlignmentLimit),
 	VolumeImportance(VolumeImportanceValue),
 	bPreserveVolume(VolumeConservation),
 	bCheckBoneBoundaries(bEnforceBoundaries),
-	MeshManager(InSrcVerts, InNumSrcVerts, InSrcIndexes, InNumSrcIndexes)
+	bMergeBonesOnCoincidentVerts(bMergeCoincidentVertBones),
+	MeshManager(InSrcVerts, InNumSrcVerts, InSrcIndexes, InNumSrcIndexes, bMergeCoincidentVertBones)
 {
 
 	// initialize the weights to be unit
@@ -624,10 +626,10 @@ int32 SkeletalSimplifier::FMeshSimplifier::CountDegenerates() const
 }
 
 
-void SkeletalSimplifier::FMeshSimplifier::OutputMesh(MeshVertType* verts, uint32* indexes, bool bMergeCoincidentVertBones, bool bWeldVtxColorAttrs, TArray<int32>* LockedVerts)
+void SkeletalSimplifier::FMeshSimplifier::OutputMesh(MeshVertType* verts, uint32* indexes, bool bWeldVtxColorAttrs, TArray<int32>* LockedVerts)
 {
 
-	if (bMergeCoincidentVertBones)
+	if(bMergeBonesOnCoincidentVerts)
 	{
 		// Fix-up to make sure that verts that share the same location (e.g. UV boundaries) have the same bone weights
 		// (otherwise we will get cracks when the characters animate)
diff --git a/Engine/Plugins/Experimental/SkeletalReduction/Source/Private/SkeletalSimplifier.h b/Engine/Plugins/Experimental/SkeletalReduction/Source/Private/SkeletalSimplifier.h
--- a/Engine/Plugins/Experimental/SkeletalReduction/Source/Private/SkeletalSimplifier.h
+++ b/Engine/Plugins/Experimental/SkeletalReduction/Source/Private/SkeletalSimplifier.h
@@ -138,7 +138,8 @@ namespace SkeletalSimplifier
 			const float CoAlignmentLimit,
 			const float VolumeImportance,
 			const bool bVolumePreservation,
-			const bool bEnforceBoundaries);
+			const bool bEnforceBoundaries,
+			const bool  bMergeCoincidentVertBones);
 
 		~FMeshSimplifier();
 
@@ -207,12 +208,11 @@ namespace SkeletalSimplifier
 		*
 		* @param Verts                 - pointer to an array to populate.  Should be least GetNumVerts() in size
 		* @param Indexes               - pointer to index buffer to populate.  Should be at least 3 * GetNumTris() in size
-		* @param bMergeCoincidentBones - if multiple verts have the same position, force them to have the same bones.
 		* @param  bWeldVtxColorAttrs   - Weld verts attributes (colors) that may have been artificially split by the reduction algorithm.
 		*                                note: this could be extended to the other attributes (e.g. normals, uvs )  
 		* @param LockedVerts           - optional pointer to array.  On return the array will hold the indices of any locked verts.
 		*/
-		void				OutputMesh(MeshVertType* Verts, uint32* Indexes, bool bMergeCoincidentVertBones = true, bool bWeldVtxColorAttrs = true, TArray<int32>* LockedVerts = NULL);
+		void				OutputMesh(MeshVertType* Verts, uint32* Indexes, bool bWeldVtxColorAttrs = true, TArray<int32>* LockedVerts = NULL);
 
 
 	protected:
@@ -286,7 +286,7 @@ namespace SkeletalSimplifier
 
 		bool                 bPreserveVolume;
 		bool                 bCheckBoneBoundaries;
-
+		bool                 bMergeBonesOnCoincidentVerts;
 		double               BoundaryConstraintWeight = 256.;
 
 
diff --git a/Engine/Plugins/Experimental/SkeletalReduction/Source/Private/SkeletalSimplifierMeshManager.cpp b/Engine/Plugins/Experimental/SkeletalReduction/Source/Private/SkeletalSimplifierMeshManager.cpp
--- a/Engine/Plugins/Experimental/SkeletalReduction/Source/Private/SkeletalSimplifierMeshManager.cpp
+++ b/Engine/Plugins/Experimental/SkeletalReduction/Source/Private/SkeletalSimplifierMeshManager.cpp
@@ -13,8 +13,9 @@
 #define INVALID_EDGE_ID UINT32_MAX
  
 SkeletalSimplifier::FSimplifierMeshManager::FSimplifierMeshManager(const MeshVertType* InSrcVerts, const uint32 InNumSrcVerts,
-	const uint32* InSrcIndexes, const uint32 InNumSrcIndexes)
+	const uint32* InSrcIndexes, const uint32 InNumSrcIndexes, const bool bMergeBonesIfSamePos)
 	:
+	bWeldBonesIfSamePos(bMergeBonesIfSamePos),
 	NumSrcVerts(InNumSrcVerts),
 	NumSrcTris(InNumSrcIndexes / 3),
 	ReducedNumVerts(InNumSrcVerts),
@@ -129,6 +130,7 @@ void SkeletalSimplifier::FSimplifierMeshManager::GroupVerts(TArray<SimpVertType>
 			HashTable.Add(HashValues[i], i);
 		}
 	}
+	const bool bTestCoincidentBones = !bWeldBonesIfSamePos;
 
 	for (int i = 0; i < NumVerts; i++)
 	{
@@ -152,8 +154,8 @@ void SkeletalSimplifier::FSimplifierMeshManager::GroupVerts(TArray<SimpVertType>
 			if (v1 == v2)
 				continue;
 
-			// link
-			if (v1->GetPos() == v2->GetPos())
+			// link if the verts have the same pos (and optionally same bones)
+			if (IsCoincident(v1, v2, bTestCoincidentBones))
 			{
 				checkSlow(v2->next == v2);
 				checkSlow(v2->prev == v2);
@@ -397,6 +399,8 @@ void SkeletalSimplifier::FSimplifierMeshManager::GroupEdges(TArray< SimpEdgeType
 		HashTable.Add(HashValues[i], i);
 	}
 
+	const bool bTestCoincidentBones = !bWeldBonesIfSamePos;
+
 	for (int32 i = 0, IMax = Edges.Num(); i < IMax; ++i)
 	{
 		// already grouped
@@ -415,13 +419,14 @@ void SkeletalSimplifier::FSimplifierMeshManager::GroupEdges(TArray< SimpEdgeType
 			if (e1 == e2)
 				continue;
 
-			bool m1 =
-				(e1->v0 == e2->v0 || e1->v0->GetPos() == e2->v0->GetPos()) &&
-				(e1->v1 == e2->v1 || e1->v1->GetPos() == e2->v1->GetPos());
-
+			bool m1 = 
+			    (e1->v0 == e2->v0  || IsCoincident(e1->v0, e2->v0, bTestCoincidentBones) ) &&
+				(e1->v1 == e2->v1  || IsCoincident(e1->v1, e2->v1, bTestCoincidentBones) );
+			// same edge but with reversed orientation
 			bool m2 =
-				(e1->v0 == e2->v1 || e1->v0->GetPos() == e2->v1->GetPos()) &&
-				(e1->v1 == e2->v0 || e1->v1->GetPos() == e2->v0->GetPos());
+				(e1->v0 == e2->v1 || IsCoincident(e1->v0, e2->v1, bTestCoincidentBones)) &&
+				(e1->v1 == e2->v0 || IsCoincident(e1->v1, e2->v0, bTestCoincidentBones));
+
 
 
 			// backwards
@@ -678,6 +683,7 @@ void SkeletalSimplifier::FSimplifierMeshManager::RebuildEdgeLinkLists(EdgePtrArr
 			HashTable.Add(HashEdgePosition(edge), i);
 		}
 
+		const bool bTestCoincidentBones = !bWeldBonesIfSamePos;
 
 		// regroup edges
 		for (uint32 i = 0; i < NumEdges; ++i)
@@ -702,19 +708,19 @@ void SkeletalSimplifier::FSimplifierMeshManager::RebuildEdgeLinkLists(EdgePtrArr
 				if (e1 == e2)
 					continue;
 
-				bool m1 =
-					(e1->v0 == e2->v0 &&
-						e1->v1 == e2->v1)
-					||
-					(e1->v0->GetPos() == e2->v0->GetPos() &&
-						e1->v1->GetPos() == e2->v1->GetPos());
-
+				bool m1 = 
+					 (e1->v0 == e2->v0 && 
+					  e1->v1 == e2->v1)
+					  ||
+					  ( IsCoincident(e1->v0, e2->v0, bTestCoincidentBones) &&
+					    IsCoincident(e1->v1, e2->v1, bTestCoincidentBones) );
+				// same edge but with reversed orientation
 				bool m2 =
 					(e1->v0 == e2->v1 &&
-						e1->v1 == e2->v0)
+					 e1->v1 == e2->v0)
 					||
-                    (e1->v0->GetPos() == e2->v1->GetPos() &&
-	                e1->v1->GetPos() == e2->v0->GetPos());
+				    ( IsCoincident(e1->v0, e2->v1, bTestCoincidentBones) &&
+					  IsCoincident(e1->v1, e2->v0, bTestCoincidentBones));
 
 				// backwards
 				if (m2)
diff --git a/Engine/Plugins/Experimental/SkeletalReduction/Source/Private/SkeletalSimplifierMeshManager.h b/Engine/Plugins/Experimental/SkeletalReduction/Source/Private/SkeletalSimplifierMeshManager.h
--- a/Engine/Plugins/Experimental/SkeletalReduction/Source/Private/SkeletalSimplifierMeshManager.h
+++ b/Engine/Plugins/Experimental/SkeletalReduction/Source/Private/SkeletalSimplifierMeshManager.h
@@ -69,7 +69,7 @@ namespace SkeletalSimplifier
 		typedef TSharedPtr<FSimplifierMeshManager>                                   Ptr;
 
 		FSimplifierMeshManager(const MeshVertType* InSrcVerts, const uint32 InNumSrcVerts,
-			const uint32* InSrcIndexes, const uint32 InNumSrcIndexes);
+			const uint32* InSrcIndexes, const uint32 InNumSrcIndexes, const bool bMergeBonesIfSamePos);
 
 		~FSimplifierMeshManager()
 		{
@@ -137,6 +137,13 @@ namespace SkeletalSimplifier
 			 return (EdgePtr->v0->adjTris.Num() == 0 || EdgePtr->v1->adjTris.Num() == 0);
 		}
 
+		// Return true if the AVert and BVert have the same vertex locations.  Optionally require the same bones.
+		bool IsCoincident(const SimpVertType* AVertPtr, const SimpVertType* BVertPtr, bool bCompairBones = true) const
+		{
+			bool bCoincident = ( AVertPtr->GetPos() == BVertPtr->GetPos()  && (!bCompairBones || AVertPtr->vert.GetSparseBones().IsApproxEquals(BVertPtr->vert.GetSparseBones())) );
+			return bCoincident;
+		}
+
 		bool IsRemoved(const SimpEdgeType* EdgePtr) const
 		{
 			return EdgePtr->TestFlags(SIMP_REMOVED);
@@ -438,6 +445,9 @@ namespace SkeletalSimplifier
 
 		}
 
+		// Weld Bones if the vertices have the same location
+		bool bWeldBonesIfSamePos = true;
+
 		// The number of verts and tris in the initial mesh.
 		const int32					NumSrcVerts = 0;
 		const int32					NumSrcTris = 0;
diff --git a/Engine/Plugins/Experimental/SkeletalReduction/Source/Private/SkeletalSimplifierVertex.h b/Engine/Plugins/Experimental/SkeletalReduction/Source/Private/SkeletalSimplifierVertex.h
--- a/Engine/Plugins/Experimental/SkeletalReduction/Source/Private/SkeletalSimplifierVertex.h
+++ b/Engine/Plugins/Experimental/SkeletalReduction/Source/Private/SkeletalSimplifierVertex.h
@@ -349,7 +349,7 @@ namespace SkeletalSimplifier
 			/**
 			* Compare two sparse arrays.  A small bone weight will be equivalent to no bone weight.
 			*/
-			bool IsApproxEquals(const BoneSparseVertexAttrs& Other, double Tolerance = (double)KINDA_SMALL_NUMBER)
+			bool IsApproxEquals(const BoneSparseVertexAttrs& Other, double Tolerance = (double)KINDA_SMALL_NUMBER) const
 			{
 
 				SparseVecD::SparseContainer Tmp = SparseData;
diff --git a/Engine/Source/Editor/DetailCustomizations/Private/SkeletalMeshReductionSettingsDetails.cpp b/Engine/Source/Editor/DetailCustomizations/Private/SkeletalMeshReductionSettingsDetails.cpp
--- a/Engine/Source/Editor/DetailCustomizations/Private/SkeletalMeshReductionSettingsDetails.cpp
+++ b/Engine/Source/Editor/DetailCustomizations/Private/SkeletalMeshReductionSettingsDetails.cpp
@@ -179,6 +179,7 @@ void FSkeletalMeshReductionSettingsDetails::CustomizeChildren(TSharedRef<IProper
 			GET_MEMBER_NAME_CHECKED(FSkeletalMeshOptimizationSettings, TerminationCriterion),
 			GET_MEMBER_NAME_CHECKED(FSkeletalMeshOptimizationSettings, bLockEdges),
 			GET_MEMBER_NAME_CHECKED(FSkeletalMeshOptimizationSettings, bEnforceBoneBoundaries),
+			GET_MEMBER_NAME_CHECKED(FSkeletalMeshOptimizationSettings, bMergeCoincidentVertBones),
 			GET_MEMBER_NAME_CHECKED(FSkeletalMeshOptimizationSettings, VolumeImportance),
 			GET_MEMBER_NAME_CHECKED(FSkeletalMeshOptimizationSettings, bLockColorBounaries)
 		};
diff --git a/Engine/Source/Editor/Persona/Private/PersonaMeshDetails.cpp b/Engine/Source/Editor/Persona/Private/PersonaMeshDetails.cpp
--- a/Engine/Source/Editor/Persona/Private/PersonaMeshDetails.cpp
+++ b/Engine/Source/Editor/Persona/Private/PersonaMeshDetails.cpp
@@ -742,6 +742,13 @@ void FSkeletalMeshReductionSettingsLayout::GenerateChildContent(IDetailChildrenB
 			FGetCheckBoxStateDelegate::CreateRaw(this, &FSkeletalMeshReductionSettingsLayout::GetEnforceBoneBoundaries),
 			FSetCheckBoxStateDelegate::CreateRaw(this, &FSkeletalMeshReductionSettingsLayout::SetEnforceBoneBoundaries));
 
+		AddBoolRow(ChildrenBuilder,
+			LOCTEXT("MergeCoincidentVertBones_Row", "MergeCoincidentVertBones"),
+			LOCTEXT("MergeCoincidentVertBones_RowNameContent", "Merge Coincident Vertices Bones"),
+			LOCTEXT("MergeCoincidentVertBones_RowNameContentTooltip", "If enabled this option make sure vertices that share the same location (e.g. UV boundaries) have the same bone weights. This can fix cracks when the characters animate."),
+			FGetCheckBoxStateDelegate::CreateRaw(this, &FSkeletalMeshReductionSettingsLayout::GetMergeCoincidentVertBones),
+			FSetCheckBoxStateDelegate::CreateRaw(this, &FSkeletalMeshReductionSettingsLayout::SetMergeCoincidentVertBones));
+
 		AddFloatRow(ChildrenBuilder,
 			LOCTEXT("VolumeImportance_Row", "Volume Importance"),
 			LOCTEXT("VolumeImportance", "Volumetric Correction"),
@@ -1383,6 +1390,20 @@ void FSkeletalMeshReductionSettingsLayout::SetEnforceBoneBoundaries(ECheckBoxSta
 	ReductionSettings.bEnforceBoneBoundaries = (NewState == ECheckBoxState::Checked) ? true : false;
 }
 
+ECheckBoxState FSkeletalMeshReductionSettingsLayout::GetMergeCoincidentVertBones() const
+{
+	return ReductionSettings.bMergeCoincidentVertBones ? ECheckBoxState::Checked : ECheckBoxState::Unchecked;
+}
+
+void FSkeletalMeshReductionSettingsLayout::SetMergeCoincidentVertBones(ECheckBoxState NewState)
+{
+	FText TransactionText = FText::Format(LOCTEXT("PersonaReductionChangedSetMergeCoincidentVertBones", "LOD{0} reduction settings: Merge Coincident Vertices Bones changed"), LODIndex);
+	FScopedTransaction Transaction(TransactionText);
+	ModifyMeshLODSettingsDelegate.ExecuteIfBound(LODIndex);
+
+	ReductionSettings.bMergeCoincidentVertBones = (NewState == ECheckBoxState::Checked) ? true : false;
+}
+
 float FSkeletalMeshReductionSettingsLayout::GetVolumeImportance() const
 {
 	return ReductionSettings.VolumeImportance;
diff --git a/Engine/Source/Editor/Persona/Private/PersonaMeshDetails.h b/Engine/Source/Editor/Persona/Private/PersonaMeshDetails.h
--- a/Engine/Source/Editor/Persona/Private/PersonaMeshDetails.h
+++ b/Engine/Source/Editor/Persona/Private/PersonaMeshDetails.h
@@ -193,6 +193,9 @@ private:
 	ECheckBoxState GetEnforceBoneBoundaries() const;
 	void SetEnforceBoneBoundaries(ECheckBoxState NewState);
 
+	ECheckBoxState GetMergeCoincidentVertBones() const;
+	void SetMergeCoincidentVertBones(ECheckBoxState NewState);
+
 	float GetVolumeImportance() const;
 	void SetVolumeImportance(float Value);
 	
diff --git a/Engine/Source/Runtime/Engine/Private/SkeletalMesh.cpp b/Engine/Source/Runtime/Engine/Private/SkeletalMesh.cpp
--- a/Engine/Source/Runtime/Engine/Private/SkeletalMesh.cpp
+++ b/Engine/Source/Runtime/Engine/Private/SkeletalMesh.cpp
@@ -5422,6 +5422,7 @@ static void SerializeReductionSettingsForDDC(FArchive& Ar, FSkeletalMeshOptimiza
 	FArchive_Serialize_BitfieldBool(Ar, ReductionSettings.bRemapMorphTargets);
 	FArchive_Serialize_BitfieldBool(Ar, ReductionSettings.bRecalcNormals);
 	FArchive_Serialize_BitfieldBool(Ar, ReductionSettings.bEnforceBoneBoundaries);
+	FArchive_Serialize_BitfieldBool(Ar, ReductionSettings.bMergeCoincidentVertBones);
 	FArchive_Serialize_BitfieldBool(Ar, ReductionSettings.bLockEdges);
 	FArchive_Serialize_BitfieldBool(Ar, ReductionSettings.bLockColorBounaries);
 	Ar << ReductionSettings.TerminationCriterion;
diff --git a/Engine/Source/Runtime/Engine/Public/SkeletalMeshReductionSettings.h b/Engine/Source/Runtime/Engine/Public/SkeletalMeshReductionSettings.h
--- a/Engine/Source/Runtime/Engine/Public/SkeletalMeshReductionSettings.h
+++ b/Engine/Source/Runtime/Engine/Public/SkeletalMeshReductionSettings.h
@@ -140,6 +140,10 @@ struct FSkeletalMeshOptimizationSettings
 	UPROPERTY(EditAnywhere, Category = FSkeletalMeshOptimizationSettings, meta = (DisplayName = "Enforce Bone Boundaries"))
 	uint8 bEnforceBoneBoundaries : 1;
 
+	/** If enabled this option make sure vertices that share the same location (e.g. UV boundaries) have the same bone weights. This can fix cracks when the characters animate. */
+	UPROPERTY(EditAnywhere, Category = FSkeletalMeshOptimizationSettings, meta = (DisplayName = "Merge Coincident Vertices Bones"))
+	uint8 bMergeCoincidentVertBones : 1;
+
 	/** Default value of 1 attempts to preserve volume.  Smaller values will loose volume by flattening curved surfaces, and larger values will accentuate curved surfaces.  */
 	UPROPERTY(EditAnywhere, Category = FSkeletalMeshOptimizationSettings, meta = (DisplayName = "Volumetric Correction", ClampMin = 0, ClampMax = 2))
 	float VolumeImportance;
@@ -189,6 +193,7 @@ struct FSkeletalMeshOptimizationSettings
 		, NormalsThreshold(60.0f)
 		, MaxBonesPerVertex(4)
 		, bEnforceBoneBoundaries(false)
+		, bMergeCoincidentVertBones(true)
 		, VolumeImportance(1.f)
 		, bLockEdges(false)
 		, bLockColorBounaries(false)
