Texcoord FVector2D -> FVector2f

diff --git a/Engine/Plugins/Experimental/SkeletalReduction/Source/Private/SkeletalSimplifierVertex.h b/Engine/Plugins/Experimental/SkeletalReduction/Source/Private/SkeletalSimplifierVertex.h
--- a/Engine/Plugins/Experimental/SkeletalReduction/Source/Private/SkeletalSimplifierVertex.h
+++ b/Engine/Plugins/Experimental/SkeletalReduction/Source/Private/SkeletalSimplifierVertex.h
@@ -41,7 +41,7 @@ namespace SkeletalSimplifier
 			FVector3f			Tangent;     // 3, 4, 5
 			FVector3f         BiTangent;   // 6, 7, 8
 			FLinearColor	Color;       // 9, 10, 11, 12
-			FVector2D		TexCoords[NumTexCoords];  // 13, .. 13 + NumTexCoords * 2 - 1
+			FVector2f		TexCoords[NumTexCoords];  // 13, .. 13 + NumTexCoords * 2 - 1
 
 
 			// used to manage identity of split/non-split vertex attributes.
