Add section to prioritize sharing the prioritize weight value. This will allow to reduce less a particular section.

diff --git a/Engine/Plugins/Experimental/SkeletalReduction/Source/Private/SkeletalMeshReductionPlugin.cpp b/Engine/Plugins/Experimental/SkeletalReduction/Source/Private/SkeletalMeshReductionPlugin.cpp
--- a/Engine/Plugins/Experimental/SkeletalReduction/Source/Private/SkeletalMeshReductionPlugin.cpp
+++ b/Engine/Plugins/Experimental/SkeletalReduction/Source/Private/SkeletalMeshReductionPlugin.cpp
@@ -158,9 +158,9 @@ private:
 	struct  FSectionRange;
 
 	/**
-	* Important bones when simplifying
+	* Important bones and sections when simplifying
 	*/
-	struct FImportantBones;
+	struct FPrioritizedFeatures;
 
 private:
 
@@ -191,7 +191,7 @@ private:
 		                         const FBoxSphereBounds& Bounds,
 		                         const FReferenceSkeleton& RefSkeleton,
 		                         FSkeletalMeshOptimizationSettings Settings,
-								 const FImportantBones& ImportantBones,
+								 const FPrioritizedFeatures& PrioritizedFeatures,
 		                         const TArray<FMatrix>& BoneMatrices,
 		                         const int32 LODIndex,
 								 const bool bReducingSourceModel,
@@ -211,11 +211,13 @@ private:
 	* @param BoneMatrices   - define the configuration of the model
 	* @param LODIndex       - target index for the result.
 	* @param OutSkinnedMesh - the posed mesh 
+	* @param VerToSectionMap - if provided, on return this array will allow VertexID to SectionID lookup
 	*/
 	void ConvertToFSkinnedSkeletalMesh( const FSkeletalMeshLODModel& SrcLODModel,
 		                                const TArray<FMatrix>& BoneMatrices,
 		                                const int32 LODIndex,
-		                                SkeletalSimplifier::FSkinnedSkeletalMesh& OutSkinnedMesh ) const;
+		                                SkeletalSimplifier::FSkinnedSkeletalMesh& OutSkinnedMesh,
+										TArray<int32>* VertToSectionMap = nullptr) const;
 
 	/**
 	* Generate a SkeletalMeshLODModel from a SkinnedSkeletalMesh and ReferenceSkeleton
@@ -298,9 +300,11 @@ private:
 	void TrimBonesPerVert( SkeletalSimplifier::FSkinnedSkeletalMesh& Mesh, int32 MaxBonesPerVert ) const ;
 
 	/**
-	* If a vertex has one of the important bones as its' major bone, associated the ImportantBones.Weight
+	* If a vertex has one of the prioritized bones as its' major bone, associated the PrioritizedFeatures.Weight
+	* Likewise if the section ID for the vert is in a prioritized section, associate the prioritized weight 
 	*/
-	void UpdateSpecializedVertWeights(const FImportantBones& ImportantBones, SkeletalSimplifier::FSkinnedSkeletalMesh& SkinnedSkeletalMesh) const;
+	void UpdateSpecializedVertWeights(const FPrioritizedFeatures& PrioritizedFeatures, const TArray<int32>& VertToSectionId, SkeletalSimplifier::FSkinnedSkeletalMesh& SkinnedSkeletalMesh) const;
+
 
 };
 
@@ -320,10 +324,16 @@ struct FQuadricSkeletalMeshReduction::FSectionRange
 	int32 End;
 };
 
-struct FQuadricSkeletalMeshReduction::FImportantBones
+struct FQuadricSkeletalMeshReduction::FPrioritizedFeatures
 {
-	TArray<int32> Ids;
+	TArray<int32> BoneIds;
+	TArray<int32> SectionIds;
 	float Weight;
+
+	bool HasFeatures() const
+	{
+		return (BoneIds.Num() >0 || SectionIds.Num() > 0);
+	}
 };
 /**
 *  Required MeshReduction Interface.
@@ -477,7 +487,8 @@ bool FQuadricSkeletalMeshReduction::RemoveMeshSection(FSkeletalMeshLODModel& Mod
 void FQuadricSkeletalMeshReduction::ConvertToFSkinnedSkeletalMesh( const FSkeletalMeshLODModel& SrcLODModel,
 	                                                               const TArray<FMatrix>& BoneMatrices,
 	                                                               const int32 LODIndex,
-	                                                               SkeletalSimplifier::FSkinnedSkeletalMesh& OutSkinnedMesh) const
+	                                                               SkeletalSimplifier::FSkinnedSkeletalMesh& OutSkinnedMesh,
+																   TArray<int32>* VertToSectionMap) const
 {
 
 
@@ -796,6 +807,30 @@ void FQuadricSkeletalMeshReduction::ConvertToFSkinnedSkeletalMesh( const FSkelet
 		}
 	}
 
+	// if requested, make a map of the vertex to original section id
+	if (VertToSectionMap != nullptr)
+	{
+		const int32 InvalidSectionId = -1;
+		VertToSectionMap->Empty();
+		VertToSectionMap->Init(InvalidSectionId, VertexCount);
+	
+		TArray<int32>& ToSection = *VertToSectionMap;
+		for (int32 s = 0; s < SectionCount; ++s)
+		{
+			if (SkipSection(s))
+			{
+				continue;
+			}
+			int32 OriginalDataSectionIndex = SrcLODModel.Sections[s].OriginalDataSectionIndex;
+
+			const FSectionRange VertexRange = SectionRangeArray[s];
+
+			for (int32 v = VertexRange.Begin; v < VertexRange.End; ++v)
+			{
+				ToSection[v] = OriginalDataSectionIndex;
+			}
+		}
+	}
 	// Put the vertex in a "correct" state.
 	//    "corrects" normals (ensures that they are orthonormal)
 	//    re-orders the bones by weight (highest to lowest)
@@ -822,16 +857,28 @@ void FQuadricSkeletalMeshReduction::ConvertToFSkinnedSkeletalMesh( const FSkelet
 }
 
 
-void FQuadricSkeletalMeshReduction::UpdateSpecializedVertWeights(const FImportantBones& ImportantBones, SkeletalSimplifier::FSkinnedSkeletalMesh& SkinnedSkeletalMesh) const
+void FQuadricSkeletalMeshReduction::UpdateSpecializedVertWeights(const FPrioritizedFeatures& PrioritizedFeatures, const TArray<int32>& VertToOriginalSectionId, SkeletalSimplifier::FSkinnedSkeletalMesh& SkinnedSkeletalMesh) const
 {
-	const float Weight = ImportantBones.Weight;
+
+	if (!PrioritizedFeatures.HasFeatures())
+	{
+		return;
+	}
+	
+	const float Weight = PrioritizedFeatures.Weight;
 
 	int32 NumVerts = SkinnedSkeletalMesh.NumVertices();
+	check(VertToOriginalSectionId.Num() == NumVerts);
 
+	const int32 InvalidSectionId = -1;
 	//If a vertex has one of the important bones as its' major bone, associated the ImportantBones.Weight
 	for (int32 i = 0; i < NumVerts; ++i)
 	{
+		const int32 OriginalSectionId = VertToOriginalSectionId[i];
 		auto& Vert = SkinnedSkeletalMesh.VertexBuffer[i];
+		Vert.SpecializedWeight = 0.f;
+
+		// Test if this vert is associated with a priority bone
 		const auto& Bones = Vert.GetSparseBones();
 		if (!Bones.bIsEmpty())
 		{
@@ -839,15 +886,19 @@ void FQuadricSkeletalMeshReduction::UpdateSpecializedVertWeights(const FImportan
 
 			const int32 FirstBone = CIter.Key(); // Bones are ordered by descending weight
 
-			if (ImportantBones.Ids.Contains(FirstBone))
+			if (PrioritizedFeatures.BoneIds.Contains(FirstBone))
 			{
 				Vert.SpecializedWeight = Weight;
 			}
 		}
-		else
+		
+		// Test if this vert is associated with a priority section
+		if (OriginalSectionId != InvalidSectionId && PrioritizedFeatures.SectionIds.Contains(OriginalSectionId))
 		{
-			Vert.SpecializedWeight = 0.f;
+			//Using += here will make bone and section prioritized verts stronger than only bone or only section prioritized verts
+			Vert.SpecializedWeight += Weight;
 		}
+
 	}
 }
 
@@ -1674,7 +1725,7 @@ bool FQuadricSkeletalMeshReduction::ReduceSkeletalLODModel( const FSkeletalMeshL
 	                                                        const FBoxSphereBounds& Bounds,
 	                                                        const FReferenceSkeleton& RefSkeleton,
 	                                                        FSkeletalMeshOptimizationSettings Settings,
-															const FImportantBones& ImportantBones,
+															const FPrioritizedFeatures& PrioritizedFeatures,
 	                                                        const TArray<FMatrix>& BoneMatrices,
 	                                                        const int32 LODIndex,
 															const bool bReducingSourceModel,
@@ -1712,9 +1763,9 @@ bool FQuadricSkeletalMeshReduction::ReduceSkeletalLODModel( const FSkeletalMeshL
 	}
 	
 	// Generate a single skinned mesh form the SrcModel.  This mesh has per-vertex tangent space.
-
+	TArray<int32> VertToOriginalSectionId; 
 	SkeletalSimplifier::FSkinnedSkeletalMesh SkinnedSkeletalMesh;
-	ConvertToFSkinnedSkeletalMesh(SrcModel, BoneMatrices, LODIndex, SkinnedSkeletalMesh);
+	ConvertToFSkinnedSkeletalMesh(SrcModel, BoneMatrices, LODIndex, SkinnedSkeletalMesh, &VertToOriginalSectionId);
 
 	int32 IterationNum = 0;
 	//We keep the original MaxNumVerts because if we iterate we want to still compare with the original request.
@@ -1723,11 +1774,10 @@ bool FQuadricSkeletalMeshReduction::ReduceSkeletalLODModel( const FSkeletalMeshL
 	{
 		if (bOptimizeMesh)
 		{
-			if (ImportantBones.Ids.Num() > 0)
-			{
-				// Add specialized weights for verts associated with "important" bones.
-				UpdateSpecializedVertWeights(ImportantBones, SkinnedSkeletalMesh);
-			}
+
+			// Add specialized weights for verts associated with "important" bones or sections
+			UpdateSpecializedVertWeights(PrioritizedFeatures, VertToOriginalSectionId, SkinnedSkeletalMesh);
+
 
 			// Capture the UV bounds from the source mesh.
 
@@ -1850,18 +1900,28 @@ void FQuadricSkeletalMeshReduction::ReduceSkeletalMesh(USkeletalMesh& SkeletalMe
 	// Struct to identify important bones.  Vertices associated with these bones
 	// will have additional collapse weight added to them.
 
-	FImportantBones  ImportantBones;
+	FPrioritizedFeatures  PrioritizedFeatures;
 	{
-		const TArray<FBoneReference>& BonesToPrioritize = LODInfo->BonesToPrioritize;
-		const float BonePrioritizationWeight = LODInfo->WeightOfPrioritization;
+		
+		const float PrioritizationWeight = LODInfo->WeightOfPrioritization;
+		PrioritizedFeatures.Weight = PrioritizationWeight;
 
-		ImportantBones.Weight = BonePrioritizationWeight;
+		const TArray<FBoneReference>& BonesToPrioritize = LODInfo->BonesToPrioritize;
 		for (const FBoneReference& BoneReference : BonesToPrioritize)
 		{
 			int32 BoneId = SkeletalMesh.GetRefSkeleton().FindRawBoneIndex(BoneReference.BoneName);
 
 			// Q: should we exclude BoneId = 0?
-			ImportantBones.Ids.AddUnique(BoneId);
+			PrioritizedFeatures.BoneIds.AddUnique(BoneId);
+		}
+
+		const TArray<FSectionReference>& SectionsToPrioritize = LODInfo->SectionsToPrioritize;
+		for (const FSectionReference& SectionReference : SectionsToPrioritize)
+		{
+			if (SectionReference.IsValidToEvaluate(SkeletalMeshResource.LODModels[LODIndex]))
+			{
+				PrioritizedFeatures.SectionIds.AddUnique(SectionReference.SectionIndex);
+			}
 		}
 	}
 	// select which mesh we're reducing from
@@ -2153,7 +2213,7 @@ void FQuadricSkeletalMeshReduction::ReduceSkeletalMesh(USkeletalMesh& SkeletalMe
 	}
 
 	// Reduce LOD model with SrcMesh if src mesh has more then 1 triangle
-	if (SrcModel->NumVertices > 3 && ReduceSkeletalLODModel(*SrcModel, *NewModel, SkeletalMesh.GetPathName(), SkeletalMesh.GetImportedBounds(), SkeletalMesh.GetRefSkeleton(), Settings, ImportantBones, RelativeToRefPoseMatrices, LODIndex, bReducingSourceModel, TargetPlatform))
+	if (SrcModel->NumVertices > 3 && ReduceSkeletalLODModel(*SrcModel, *NewModel, SkeletalMesh.GetPathName(), SkeletalMesh.GetImportedBounds(), SkeletalMesh.GetRefSkeleton(), Settings, PrioritizedFeatures, RelativeToRefPoseMatrices, LODIndex, bReducingSourceModel, TargetPlatform))
 	{
 		FSkeletalMeshLODInfo* ReducedLODInfoPtr = SkeletalMesh.GetLODInfo(LODIndex);
 		check(ReducedLODInfoPtr);
diff --git a/Engine/Source/Editor/Persona/Private/PersonaMeshDetails.cpp b/Engine/Source/Editor/Persona/Private/PersonaMeshDetails.cpp
--- a/Engine/Source/Editor/Persona/Private/PersonaMeshDetails.cpp
+++ b/Engine/Source/Editor/Persona/Private/PersonaMeshDetails.cpp
@@ -2477,6 +2477,12 @@ void FPersonaMeshDetails::CustomizeLODInfoSetingsDetails(IDetailLayoutBuilder& D
 	PreventAttributePostEditChange(BonesToPrioritizeHandle, LODIndex, BonesToPrioritizeHandle->GetProperty()->GetFName(), bForceComponentRefreshFalse);
 	DetailLayout.HideProperty(BonesToPrioritizeHandle);
 
+	TSharedPtr<IPropertyHandle> SectionsToPrioritizeHandle = LODInfoChild->GetChildHandle(GET_MEMBER_NAME_CHECKED(FSkeletalMeshLODInfo, SectionsToPrioritize));
+	IDetailPropertyRow& SectionsToPrioritizeRow = LODInfoGroup.AddPropertyRow(SectionsToPrioritizeHandle->AsShared());
+	SectionsToPrioritizeRow.IsEnabled(EnabledAttrib);
+	PreventAttributePostEditChange(SectionsToPrioritizeHandle, LODIndex, SectionsToPrioritizeHandle->GetProperty()->GetFName(), bForceComponentRefreshFalse);
+	DetailLayout.HideProperty(SectionsToPrioritizeHandle);
+
 	TSharedPtr<IPropertyHandle> WeightToPriortizeHandle = LODInfoChild->GetChildHandle(GET_MEMBER_NAME_CHECKED(FSkeletalMeshLODInfo, WeightOfPrioritization));
 	IDetailPropertyRow& WeightToPriortizeRow = LODInfoGroup.AddPropertyRow(WeightToPriortizeHandle->AsShared());
 	WeightToPriortizeRow.IsEnabled(EnabledAttrib);
@@ -2484,7 +2490,8 @@ void FPersonaMeshDetails::CustomizeLODInfoSetingsDetails(IDetailLayoutBuilder& D
 	DetailLayout.HideProperty(WeightToPriortizeHandle);
 
 	const TArray<FName> HiddenProperties = { GET_MEMBER_NAME_CHECKED(FSkeletalMeshLODInfo, ReductionSettings), GET_MEMBER_NAME_CHECKED(FSkeletalMeshLODInfo, BakePose), GET_MEMBER_NAME_CHECKED(FSkeletalMeshLODInfo, BakePoseOverride), GET_MEMBER_NAME_CHECKED(FSkeletalMeshLODInfo, BonesToRemove),
-		GET_MEMBER_NAME_CHECKED(FSkeletalMeshLODInfo, BonesToPrioritize), GET_MEMBER_NAME_CHECKED(FSkeletalMeshLODInfo, WeightOfPrioritization), GET_MEMBER_NAME_CHECKED(FSkeletalMeshLODInfo, ScreenSize), GET_MEMBER_NAME_CHECKED(FSkeletalMeshLODInfo, LODHysteresis), GET_MEMBER_NAME_CHECKED(FSkeletalMeshLODInfo, BuildSettings) };
+		GET_MEMBER_NAME_CHECKED(FSkeletalMeshLODInfo, BonesToPrioritize), GET_MEMBER_NAME_CHECKED(FSkeletalMeshLODInfo, SectionsToPrioritize), GET_MEMBER_NAME_CHECKED(FSkeletalMeshLODInfo, WeightOfPrioritization), GET_MEMBER_NAME_CHECKED(FSkeletalMeshLODInfo, ScreenSize),
+		GET_MEMBER_NAME_CHECKED(FSkeletalMeshLODInfo, LODHysteresis), GET_MEMBER_NAME_CHECKED(FSkeletalMeshLODInfo, BuildSettings) };
 	for (uint32 ChildIndex = 0; ChildIndex < NumInfoChildren; ++ChildIndex)
 	{
 		TSharedRef<IPropertyHandle> LODInfoChildHandle = LODInfoChild->GetChildHandle(ChildIndex).ToSharedRef();
diff --git a/Engine/Source/Editor/Persona/Private/PersonaModule.cpp b/Engine/Source/Editor/Persona/Private/PersonaModule.cpp
--- a/Engine/Source/Editor/Persona/Private/PersonaModule.cpp
+++ b/Engine/Source/Editor/Persona/Private/PersonaModule.cpp
@@ -92,6 +92,7 @@
 #include "AnimSequenceTimelineCommands.h"
 #include "SAnimMontageSectionsPanel.h"
 #include "Animation/AnimSequenceHelpers.h"
+#include "SkeletalMeshReferenceSectionDetails.h"
 
 IMPLEMENT_MODULE( FPersonaModule, Persona );
 
@@ -133,6 +134,8 @@ void FPersonaModule::StartupModule()
 
 		PropertyModule.RegisterCustomPropertyTypeLayout("SkeletalMeshSamplingRegionBoneFilter", FOnGetPropertyTypeCustomizationInstance::CreateStatic(&FNiagaraSkeletalMeshRegionBoneFilterDetails::MakeInstance));
 		PropertyModule.RegisterCustomPropertyTypeLayout("SkeletalMeshSamplingRegionMaterialFilter", FOnGetPropertyTypeCustomizationInstance::CreateStatic(&FNiagaraSkeletalMeshRegionMaterialFilterDetails::MakeInstance));
+
+		PropertyModule.RegisterCustomPropertyTypeLayout("SectionReference", FOnGetPropertyTypeCustomizationInstance::CreateStatic(&FSectionReferenceCustomization::MakeInstance));
 	}
 
 	// Register the editor modes
diff --git a/Engine/Source/Editor/Persona/Private/Shared/ReferenceSectionSelectionWidget.cpp b/Engine/Source/Editor/Persona/Private/Shared/ReferenceSectionSelectionWidget.cpp
new file mode 100644
--- /dev/null
+++ b/Engine/Source/Editor/Persona/Private/Shared/ReferenceSectionSelectionWidget.cpp
@@ -0,0 +1,190 @@
+// Copyright Epic Games, Inc. All Rights Reserved.
+
+
+#include "ReferenceSectionSelectionWidget.h"
+
+#include "DetailLayoutBuilder.h"
+#include "EditorStyleSet.h"
+#include "Rendering/SkeletalMeshLODModel.h"
+#include "Widgets/Input/SComboButton.h"
+#include "Widgets/Input/SSearchBox.h"
+#include "Widgets/Layout/SSeparator.h"
+
+#define LOCTEXT_NAMESPACE "SReferenceSectionSelectionWidget"
+
+void SReferenceSectionSelectionWidget::Construct(const FArguments& InArgs)
+{
+	OnSectionSelectionChanged = InArgs._OnSectionSelectionChanged;
+	OnGetSelectedSection = InArgs._OnGetSelectedSection;
+	OnGetLodModel = InArgs._OnGetLodModel;
+	bHideChunkedSections = InArgs._bHideChunkedSections;
+
+	this->ChildSlot
+	[
+		SAssignNew(SectionPickerButton, SComboButton)
+		.OnGetMenuContent(FOnGetContent::CreateSP(this, &SReferenceSectionSelectionWidget::CreateSectionListWidgetMenu))
+		.ContentPadding(FMargin(4.0f, 2.0f, 4.0f, 2.0f))
+		.ButtonContent()
+		[
+			SNew(STextBlock)
+			.Text(this, &SReferenceSectionSelectionWidget::GetCurrentSectionIndex)
+			.Font(IDetailLayoutBuilder::GetDetailFont())
+		]
+	];
+}
+
+const TArray<TSharedPtr<int32>>* SReferenceSectionSelectionWidget::GetSections() const
+{
+	CacheSectionList.Reset();
+	if (OnGetLodModel.IsBound())
+	{
+		const FSkeletalMeshLODModel& LodModel = OnGetLodModel.Execute();
+		for (int32 SectionIndex = 0; SectionIndex < LodModel.Sections.Num(); ++SectionIndex)
+		{
+			if (LodModel.Sections[SectionIndex].ChunkedParentSectionIndex != INDEX_NONE && bHideChunkedSections)
+			{
+				continue;
+			}
+			CacheSectionList.Add(MakeShareable(new int32(SectionIndex)));
+		}
+	}
+	return &CacheSectionList;
+}
+
+TSharedRef<ITableRow> SReferenceSectionSelectionWidget::MakeIntegerDisplayWidget(const TSharedPtr<int32> SectionIndex, const TSharedRef<STableViewBase>& OwnerTable) const
+{
+	const int32 SectionIndexValue = *SectionIndex;
+	FString ItemString = FString::FromInt(SectionIndexValue);
+	return SNew(SComboRow< TSharedRef<FString> >, OwnerTable)
+		[
+			SNew(SHorizontalBox)
+			+ SHorizontalBox::Slot()
+			.FillWidth(1)
+			.Padding(2.0f)
+			[
+				SNew(STextBlock)
+				.Text(FText::FromString(ItemString))
+				.Font(FEditorStyle::GetFontStyle(TEXT("PropertyWindow.NormalFont")))
+			]
+		];
+}
+
+TSharedRef<SWidget> SReferenceSectionSelectionWidget::CreateSectionListWidgetMenu()
+{
+	bool bMultipleValues = false;
+	int32 CurrentSectionIndex;
+	if (OnGetSelectedSection.IsBound())
+	{
+		CurrentSectionIndex = OnGetSelectedSection.Execute(bMultipleValues);
+	}
+
+	TSharedPtr<SHeaderRow> PlatformColumnHeader = SNew(SHeaderRow);
+
+
+	//Create a list widget containing all possible section we can reference
+	SAssignNew(SectionListView, SListView<TSharedPtr<int32>>)
+		.ListItemsSource(GetSections())
+		.OnGenerateRow(this, &SReferenceSectionSelectionWidget::MakeIntegerDisplayWidget)
+		.Visibility(EVisibility::Visible)
+		.SelectionMode(ESelectionMode::Single)
+		.OnSelectionChanged(this, &SReferenceSectionSelectionWidget::OnSelectionChanged)
+		.HeaderRow(PlatformColumnHeader.ToSharedRef());
+
+	FText TitleText = FText(LOCTEXT("ReferenceListWidgetMenuTitle", "Choose a Section"));
+	TSharedPtr<SBorder> BorderWidget;
+	SAssignNew(BorderWidget, SBorder)
+	.Padding(6)
+	.BorderImage(FEditorStyle::GetBrush("NoBorder"))
+	.Content()
+	[
+		SNew(SBox)
+		.WidthOverride(300)
+		.HeightOverride(512)
+		.Content()
+		[
+			SNew(SVerticalBox)
+			+ SVerticalBox::Slot()
+			.AutoHeight()
+			[
+				SNew(STextBlock)
+				.Font(FEditorStyle::GetFontStyle("BoldFont"))
+				.Text(TitleText)
+			]
+			+ SVerticalBox::Slot()
+			.AutoHeight()
+			[
+				SNew(SSeparator)
+				.SeparatorImage(FEditorStyle::GetBrush("Menu.Separator"))
+				.Orientation(Orient_Horizontal)
+			]
+			+ SVerticalBox::Slot()
+			[
+				SectionListView->AsShared()
+			]
+		]
+	];
+
+	return BorderWidget.ToSharedRef();
+}
+
+void SReferenceSectionSelectionWidget::OnSelectionChanged(TSharedPtr<int32> NewSectionIndex, ESelectInfo::Type SelectInfo)
+{
+	const int32 SectionIndex = *NewSectionIndex;
+	//Because we recreate all our items on tree refresh we will get a spurious null selection event initially.
+	if (OnSectionSelectionChanged.IsBound() && OnGetLodModel.IsBound())
+	{
+		const FSkeletalMeshLODModel& LodModel = OnGetLodModel.Execute();
+		int32 SaveSectionIndex = SectionIndex;
+		if (LodModel.Sections.IsValidIndex(SectionIndex))
+		{
+			SaveSectionIndex = LodModel.Sections[SectionIndex].OriginalDataSectionIndex;
+		}
+		OnSectionSelectionChanged.Execute(SaveSectionIndex);
+	}
+
+	SectionPickerButton->SetIsOpen(false);
+}
+
+FText SReferenceSectionSelectionWidget::GetCurrentSectionIndex() const
+{
+	if(OnGetSelectedSection.IsBound())
+	{
+		bool bMultipleValues = false;
+		int32 CurrentSectionIndexValue = OnGetSelectedSection.Execute(bMultipleValues);
+		if(bMultipleValues)
+		{
+			return LOCTEXT("MultipleValues", "Multiple Values");
+		}
+		else
+		{
+			int32 DisplaySectionIndex = CurrentSectionIndexValue;
+			if (bHideChunkedSections && OnGetLodModel.IsBound())
+			{
+				const FSkeletalMeshLODModel& LodModel = OnGetLodModel.Execute();
+				for (int32 SectionIndex = 0; SectionIndex < LodModel.Sections.Num(); ++SectionIndex)
+				{
+					if (LodModel.Sections[SectionIndex].ChunkedParentSectionIndex != INDEX_NONE && bHideChunkedSections)
+					{
+						continue;
+					}
+					if (LodModel.Sections[SectionIndex].OriginalDataSectionIndex == CurrentSectionIndexValue)
+					{
+						DisplaySectionIndex = SectionIndex;
+						break;
+					}
+				}
+			}
+			if (DisplaySectionIndex == INDEX_NONE)
+			{
+				return LOCTEXT("GetCurrentSectionIndexTextNone", "Section None");
+			}
+			return FText::Format(LOCTEXT("GetCurrentSectionIndexText", "Section {0}"), DisplaySectionIndex);
+		}
+	}
+
+	// @todo implement default solution?
+	return FText::GetEmpty();
+}
+
+#undef LOCTEXT_NAMESPACE
+
diff --git a/Engine/Source/Editor/Persona/Private/SkeletalMeshReferenceSectionDetails.cpp b/Engine/Source/Editor/Persona/Private/SkeletalMeshReferenceSectionDetails.cpp
new file mode 100644
--- /dev/null
+++ b/Engine/Source/Editor/Persona/Private/SkeletalMeshReferenceSectionDetails.cpp
@@ -0,0 +1,129 @@
+// Copyright Epic Games, Inc. All Rights Reserved.
+
+#include "SkeletalMeshReferenceSectionDetails.h"
+
+#include "DetailCategoryBuilder.h"
+#include "DetailLayoutBuilder.h"
+#include "EditorStyleSet.h"
+#include "Engine/SkeletalMesh.h"
+#include "IDetailCustomization.h"
+#include "IDetailPropertyRow.h"
+#include "IDetailsView.h"
+#include "IPropertyTypeCustomization.h"
+#include "DetailWidgetRow.h"
+#include "Modules/ModuleManager.h"
+#include "PropertyHandle.h"
+#include "ReferenceSectionSelectionWidget.h"
+#include "Rendering/SkeletalMeshLODModel.h"
+#include "Rendering/SkeletalMeshModel.h"
+#include "Templates/SharedPointer.h"
+#include "Widgets/Input/SEditableTextBox.h"
+#include "Widgets/SBoxPanel.h"
+#include "Widgets/Text/STextBlock.h"
+
+
+//Private persona include
+#include "LODInfoUILayout.h"
+
+TSharedRef<IPropertyTypeCustomization> FSectionReferenceCustomization::MakeInstance()
+{
+	return MakeShareable(new FSectionReferenceCustomization());
+}
+
+void FSectionReferenceCustomization::CustomizeHeader(TSharedRef<IPropertyHandle> StructPropertyHandle, class FDetailWidgetRow& HeaderRow, IPropertyTypeCustomizationUtils& StructCustomizationUtils)
+{
+	// set property handle 
+	SetPropertyHandle(StructPropertyHandle);
+	// set editable LodModel info from struct
+	SetEditableLodModel(StructPropertyHandle);
+	if (TargetSkeletalMesh != nullptr && TargetLodIndex != INDEX_NONE && SectionIndexProperty->IsValidHandle())
+	{
+		HeaderRow
+		.NameContent()
+		[
+			StructPropertyHandle->CreatePropertyNameWidget()
+		]
+		.ValueContent()
+		.MaxDesiredWidth(0.0f)
+		[
+			SNew(SReferenceSectionSelectionWidget)
+			.ToolTipText(StructPropertyHandle->GetToolTipText())
+			.OnSectionSelectionChanged(this, &FSectionReferenceCustomization::OnSectionSelectionChanged)
+			.OnGetSelectedSection(this, &FSectionReferenceCustomization::GetSelectedSection)
+			.OnGetLodModel(this, &FSectionReferenceCustomization::GetLodModel)
+			.bHideChunkedSections(true)
+		];
+	}
+	else
+	{
+		const bool bEnsureOnInvalidLodModel = true;
+		ensureAlways(!bEnsureOnInvalidLodModel);
+	}
+}
+
+void FSectionReferenceCustomization::CustomizeChildren(TSharedRef<IPropertyHandle> StructPropertyHandle, class IDetailChildrenBuilder& ChildBuilder, IPropertyTypeCustomizationUtils& StructCustomizationUtils)
+{
+
+}
+
+void FSectionReferenceCustomization::SetEditableLodModel(TSharedRef<IPropertyHandle> StructPropertyHandle)
+{
+	TArray<UObject*> Objects;
+	StructPropertyHandle->GetOuterObjects(Objects);
+
+	TargetSkeletalMesh = nullptr;
+	TargetLodIndex = INDEX_NONE;
+
+	for (UObject* Outer : Objects)
+	{
+		//We customize only when we have a ULODInfoUILayout has the parent because we need the lod index to customize the UI
+		if (ULODInfoUILayout* LODInfoUILayout = Cast<ULODInfoUILayout>(Outer))
+		{
+			TargetSkeletalMesh = LODInfoUILayout->GetPersonaToolkit()->GetPreviewMesh();
+			check(TargetSkeletalMesh);
+			TargetLodIndex = LODInfoUILayout->GetLODIndex();
+			break;
+		}
+	}
+}
+
+TSharedPtr<IPropertyHandle> FSectionReferenceCustomization::FindStructMemberProperty(TSharedRef<IPropertyHandle> PropertyHandle, const FName& PropertyName)
+{
+	uint32 NumChildren = 0;
+	PropertyHandle->GetNumChildren(NumChildren);
+	for (uint32 ChildIdx = 0; ChildIdx < NumChildren; ++ChildIdx)
+	{
+		TSharedPtr<IPropertyHandle> ChildHandle = PropertyHandle->GetChildHandle(ChildIdx);
+		if (ChildHandle->GetProperty()->GetFName() == PropertyName)
+		{
+			return ChildHandle;
+		}
+	}
+
+	return TSharedPtr<IPropertyHandle>();
+}
+
+void FSectionReferenceCustomization::SetPropertyHandle(TSharedRef<IPropertyHandle> StructPropertyHandle)
+{
+	SectionIndexProperty = FindStructMemberProperty(StructPropertyHandle, GET_MEMBER_NAME_CHECKED(FSectionReference, SectionIndex));
+	check(SectionIndexProperty->IsValidHandle());
+}
+
+void FSectionReferenceCustomization::OnSectionSelectionChanged(int32 SectionIndex)
+{
+	SectionIndexProperty->SetValue(SectionIndex);
+}
+
+int32 FSectionReferenceCustomization::GetSelectedSection(bool& bMultipleValues) const
+{
+	FString OutText;
+	FPropertyAccess::Result Result = SectionIndexProperty->GetValueAsFormattedString(OutText);
+	bMultipleValues = (Result == FPropertyAccess::MultipleValues);
+	return FCString::Atoi(*OutText);
+}
+
+const FSkeletalMeshLODModel& FSectionReferenceCustomization::GetLodModel() const
+{
+	check(TargetSkeletalMesh && TargetLodIndex != INDEX_NONE && TargetSkeletalMesh->GetImportedModel() && TargetSkeletalMesh->GetImportedModel()->LODModels.IsValidIndex(TargetLodIndex));
+	return TargetSkeletalMesh->GetImportedModel()->LODModels[TargetLodIndex];
+}
diff --git a/Engine/Source/Editor/Persona/Private/SkeletalMeshReferenceSectionDetails.h b/Engine/Source/Editor/Persona/Private/SkeletalMeshReferenceSectionDetails.h
new file mode 100644
--- /dev/null
+++ b/Engine/Source/Editor/Persona/Private/SkeletalMeshReferenceSectionDetails.h
@@ -0,0 +1,41 @@
+// Copyright Epic Games, Inc. All Rights Reserved.
+
+#pragma once
+
+#include "CoreMinimal.h"
+#include "IDetailCustomization.h"
+#include "Types/SlateEnums.h"
+
+class FSkeletalMeshLODModel;
+class IDetailLayoutBuilder;
+class IPropertyHandle;
+class SEditableTextBox;
+class USkeletalMesh;
+
+class FSectionReferenceCustomization : public IPropertyTypeCustomization
+{
+public:
+	static TSharedRef<IPropertyTypeCustomization> MakeInstance();
+
+	// IPropertyTypeCustomization interface
+	virtual void CustomizeHeader(TSharedRef<IPropertyHandle> StructPropertyHandle, class FDetailWidgetRow& HeaderRow, IPropertyTypeCustomizationUtils& StructCustomizationUtils) override;
+	virtual void CustomizeChildren(TSharedRef<IPropertyHandle> StructPropertyHandle, class IDetailChildrenBuilder& ChildBuilder, IPropertyTypeCustomizationUtils& StructCustomizationUtils) override;
+	const FSkeletalMeshLODModel& GetLodModel() const;
+
+protected:
+	void SetEditableLodModel(TSharedRef<IPropertyHandle> StructPropertyHandle);
+	virtual void SetPropertyHandle(TSharedRef<IPropertyHandle> StructPropertyHandle);
+	TSharedPtr<IPropertyHandle> FindStructMemberProperty(TSharedRef<IPropertyHandle> PropertyHandle, const FName& PropertyName);
+	// Property to change after section index has been picked
+	TSharedPtr<IPropertyHandle> SectionIndexProperty;
+
+	// Target LodModel this widget is referencing
+	USkeletalMesh* TargetSkeletalMesh = nullptr;
+	int32 TargetLodIndex = INDEX_NONE;
+
+private:
+
+	// Section widget delegates
+	virtual void OnSectionSelectionChanged(int32 SectionIndex);
+	virtual int32 GetSelectedSection(bool& bMultipleValues) const;
+};
diff --git a/Engine/Source/Editor/Persona/Public/ReferenceSectionSelectionWidget.h b/Engine/Source/Editor/Persona/Public/ReferenceSectionSelectionWidget.h
new file mode 100644
--- /dev/null
+++ b/Engine/Source/Editor/Persona/Public/ReferenceSectionSelectionWidget.h
@@ -0,0 +1,78 @@
+// Copyright Epic Games, Inc. All Rights Reserved.
+
+#pragma once
+
+#include "CoreMinimal.h"
+#include "SlateFwd.h"
+#include "Types/SlateEnums.h"
+#include "Widgets/DeclarativeSyntaxSupport.h"
+#include "Widgets/SWidget.h"
+#include "Widgets/SCompoundWidget.h"
+#include "Widgets/Views/STableViewBase.h"
+#include "Widgets/Views/STableRow.h"
+#include "Widgets/Views/STreeView.h"
+#include "Widgets/Views/SListView.h"
+
+class FSkeletalMeshLODModel;
+class IEditableSkeleton;
+class SComboButton;
+
+DECLARE_DELEGATE_OneParam(FOnSectionSelectionChanged, int32);
+DECLARE_DELEGATE_RetVal_OneParam(int32, FGetSelectedSection, bool& /*bMultipleValues*/);
+DECLARE_DELEGATE_RetVal(const FSkeletalMeshLODModel&, FGetLodModel);
+
+class PERSONA_API SReferenceSectionSelectionWidget : public SCompoundWidget
+{
+public:
+
+	SLATE_BEGIN_ARGS(SReferenceSectionSelectionWidget)
+		: _OnSectionSelectionChanged()
+		, _OnGetSelectedSection()
+		, _OnGetLodModel()
+	{}
+
+		/** Should we hide chunked sections when showing the sections list we can select */
+		SLATE_ARGUMENT(bool, bHideChunkedSections)
+
+		/** set selected section index */
+		SLATE_EVENT(FOnSectionSelectionChanged, OnSectionSelectionChanged);
+
+		/** get selected section index **/
+		SLATE_EVENT(FGetSelectedSection, OnGetSelectedSection);
+
+		SLATE_EVENT(FGetLodModel, OnGetLodModel)
+
+	SLATE_END_ARGS();
+
+	/**
+	 * Construct this widget
+	 *
+	 * @param	InArgs	The declaration data for this widget
+	 */
+	void Construct( const FArguments& InArgs );
+
+private: 
+
+	// Creates the combo button menu when clicked
+	TSharedRef<SWidget> CreateSectionListWidgetMenu();
+	TSharedRef<ITableRow> MakeIntegerDisplayWidget(const TSharedPtr<int32> SectionIndex, const TSharedRef<STableViewBase>& OwnerTable) const;
+
+	TSharedPtr<SListView<TSharedPtr<int32>>> SectionListView;
+	// Called when the user selects a bone name
+	void OnSelectionChanged(TSharedPtr<int32> NewSectionIndex, ESelectInfo::Type SelectInfo);
+
+	const TArray<TSharedPtr<int32>>* GetSections() const;
+
+	FText GetCurrentSectionIndex() const;
+
+	// Base combo button 
+	TSharedPtr<SComboButton> SectionPickerButton;
+
+	// delegates
+	FOnSectionSelectionChanged OnSectionSelectionChanged;
+	FGetSelectedSection		OnGetSelectedSection;
+	FGetLodModel	OnGetLodModel;
+
+	bool bHideChunkedSections = false;
+	mutable TArray<TSharedPtr<int32>> CacheSectionList;
+};
diff --git a/Engine/Source/Editor/UnrealEd/Private/ImportUtils/SkeletalMeshImportUtils.cpp b/Engine/Source/Editor/UnrealEd/Private/ImportUtils/SkeletalMeshImportUtils.cpp
--- a/Engine/Source/Editor/UnrealEd/Private/ImportUtils/SkeletalMeshImportUtils.cpp
+++ b/Engine/Source/Editor/UnrealEd/Private/ImportUtils/SkeletalMeshImportUtils.cpp
@@ -461,6 +461,7 @@ void SkeletalMesUtilsImpl::RestoreLODInfo(const TSharedPtr<const FExistingSkelMe
 	}
 	ImportedLODInfo.BonesToRemove = ExistingLODInfo.BonesToRemove;
 	ImportedLODInfo.BonesToPrioritize = ExistingLODInfo.BonesToPrioritize;
+	ImportedLODInfo.SectionsToPrioritize = ExistingLODInfo.SectionsToPrioritize;
 	ImportedLODInfo.WeightOfPrioritization = ExistingLODInfo.WeightOfPrioritization;
 	ImportedLODInfo.BakePose = ExistingLODInfo.BakePose;
 	ImportedLODInfo.BakePoseOverride = ExistingLODInfo.BakePoseOverride;
diff --git a/Engine/Source/Runtime/Engine/Classes/Engine/SkeletalMesh.h b/Engine/Source/Runtime/Engine/Classes/Engine/SkeletalMesh.h
--- a/Engine/Source/Runtime/Engine/Classes/Engine/SkeletalMesh.h
+++ b/Engine/Source/Runtime/Engine/Classes/Engine/SkeletalMesh.h
@@ -359,6 +359,51 @@ struct ENGINE_API FSkeletalMeshClothBuildParams
 	TSoftObjectPtr<UPhysicsAsset> PhysicsAsset;
 };
 
+USTRUCT()
+struct FSectionReference
+{
+	GENERATED_USTRUCT_BODY()
+
+	/** Index of the section we reference. **/
+	UPROPERTY(EditAnywhere, Category = SectionReference)
+	int32 SectionIndex;
+
+	FSectionReference()
+		: SectionIndex(INDEX_NONE)
+	{
+	}
+
+	FSectionReference(const int32& InSectionIndex)
+		: SectionIndex(InSectionIndex)
+	{
+	}
+
+	bool operator==(const FSectionReference& Other) const
+	{
+		return SectionIndex == Other.SectionIndex;
+	}
+
+#if WITH_EDITOR
+	/** return true if it has a valid section index for LodModel parameter **/
+	ENGINE_API bool IsValidToEvaluate(const FSkeletalMeshLODModel& LodModel) const;
+
+	const struct FSkelMeshSection* GetMeshLodSection(const FSkeletalMeshLODModel& LodModel) const;
+	int32 GetMeshLodSectionIndex(const FSkeletalMeshLODModel& LodModel) const;
+#endif
+
+	bool Serialize(FArchive& Ar)
+	{
+		Ar << SectionIndex;
+		return true;
+	}
+
+	friend FArchive& operator<<(FArchive& Ar, FSectionReference& B)
+	{
+		B.Serialize(Ar);
+		return Ar;
+	}
+};
+
 /** Struct containing information for a particular LOD level, such as materials and info for when to use it. */
 USTRUCT()
 struct FSkeletalMeshLODInfo
@@ -406,11 +451,15 @@ struct FSkeletalMeshLODInfo
 	UPROPERTY(EditAnywhere, Category = ReductionSettings)
 	TArray<FBoneReference> BonesToRemove;
 
-	/** Bones which should be prioritized for the quality, this will be weighted toward keeping source data. */
+	/** Bones which should be prioritized for the quality, this will be weighted toward keeping source data. Use WeightOfPrioritization to control the value. */
 	UPROPERTY(EditAnywhere, Category = ReductionSettings)
 	TArray<FBoneReference> BonesToPrioritize;
 
-	/** How much to consideration to give BonesToPrioritize.  The weight is an additional vertex simplification penalty where 0 means nothing. */
+	/** Sections which should be prioritized for the quality, this will be weighted toward keeping source data. Use WeightOfPrioritization to control the value. */
+	UPROPERTY(EditAnywhere, Category = ReductionSettings)
+	TArray<FSectionReference> SectionsToPrioritize;
+
+	/** How much to consideration to give BonesToPrioritize and SectionsToPrioritize.  The weight is an additional vertex simplification penalty where 0 means nothing. */
 	UPROPERTY(EditAnywhere, Category = ReductionSettings, meta = (UIMin = "0.0", ClampMin = "0.0"))
 	float WeightOfPrioritization;
 
diff --git a/Engine/Source/Runtime/Engine/Classes/Engine/SkeletalMeshLODSettings.h b/Engine/Source/Runtime/Engine/Classes/Engine/SkeletalMeshLODSettings.h
--- a/Engine/Source/Runtime/Engine/Classes/Engine/SkeletalMeshLODSettings.h
+++ b/Engine/Source/Runtime/Engine/Classes/Engine/SkeletalMeshLODSettings.h
@@ -93,12 +93,16 @@ struct FSkeletalMeshLODGroupSettings
 	UPROPERTY(EditAnywhere, Category = Reduction)
 	TArray<FBoneFilter> BoneList;
 
-	/** Bones which should be prioritized for the quality, this will be weighted toward keeping source data. */
+	/** Bones which should be prioritized for the quality, this will be weighted toward keeping source data. Use WeightOfPrioritization to control the value. */
 	UPROPERTY(EditAnywhere, Category = Reduction)
 	TArray<FName> BonesToPrioritize;
 
-	/** Weight of how much consider for BonesToPrioritize. 0 means nothing, and 1 means take all source */
-	UPROPERTY(EditAnywhere, Category = Reduction, meta = (UIMin = "0.0", UIMax = "1.0", ClampMin = "0.0", ClampMax = "1.0"))
+	/** Sections which should be prioritized for the quality, this will be weighted toward keeping source data. Use WeightOfPrioritization to control the value. */
+	UPROPERTY(EditAnywhere, Category = Reduction)
+	TArray<int32> SectionsToPrioritize;
+
+	/** How much to consideration to give BonesToPrioritize and SectionsToPrioritize.  The weight is an additional vertex simplification penalty where 0 means nothing. */
+	UPROPERTY(EditAnywhere, Category = Reduction, meta = (UIMin = "0.0", ClampMin = "0.0"))
 	float WeightOfPrioritization;
 
 	/** Pose which should be used to reskin vertex influences for which the bones will be removed in this LOD level, uses ref-pose by default */
diff --git a/Engine/Source/Runtime/Engine/Private/SkeletalMesh.cpp b/Engine/Source/Runtime/Engine/Private/SkeletalMesh.cpp
--- a/Engine/Source/Runtime/Engine/Private/SkeletalMesh.cpp
+++ b/Engine/Source/Runtime/Engine/Private/SkeletalMesh.cpp
@@ -440,13 +440,47 @@ void FScopedSkeletalMeshPostEditChange::SetSkeletalMesh(USkeletalMesh* InSkeleta
 	}
 }
 
-#endif //WITH_EDITOR
+
+/*-----------------------------------------------------------------------------
+	FSectionReference
+-----------------------------------------------------------------------------*/
+
+bool FSectionReference::IsValidToEvaluate(const FSkeletalMeshLODModel& LodModel) const
+{
+	return GetMeshLodSection(LodModel) != nullptr;
+}
+
+const FSkelMeshSection* FSectionReference::GetMeshLodSection(const FSkeletalMeshLODModel& LodModel) const
+{
+	for (int32 LodModelSectionIndex = 0; LodModelSectionIndex < LodModel.Sections.Num(); ++LodModelSectionIndex)
+	{
+		const FSkelMeshSection& Section = LodModel.Sections[LodModelSectionIndex];
+		if (Section.ChunkedParentSectionIndex == INDEX_NONE && SectionIndex == Section.OriginalDataSectionIndex)
+		{
+			return &Section;
+		}
+	}
+	return nullptr;
+}
+
+int32 FSectionReference::GetMeshLodSectionIndex(const FSkeletalMeshLODModel& LodModel) const
+{
+	for (int32 LodModelSectionIndex = 0; LodModelSectionIndex < LodModel.Sections.Num(); ++LodModelSectionIndex)
+	{
+		const FSkelMeshSection& Section = LodModel.Sections[LodModelSectionIndex];
+		if (Section.ChunkedParentSectionIndex == INDEX_NONE && SectionIndex == Section.OriginalDataSectionIndex)
+		{
+			return LodModelSectionIndex;
+		}
+	}
+	return INDEX_NONE;
+}
+
 
 /*-----------------------------------------------------------------------------
 	FSkeletalMeshAsyncBuildWorker
 -----------------------------------------------------------------------------*/
 
-#if WITH_EDITOR
 
 void FSkeletalMeshAsyncBuildWorker::DoWork()
 {
@@ -5283,6 +5317,7 @@ FSkeletalMeshLODInfo& USkeletalMesh::AddLODInfo()
 			NewLODInfo.BakePoseOverride = LODInfoArray[LastIndex].BakePoseOverride;
 			NewLODInfo.BonesToRemove = LODInfoArray[LastIndex].BonesToRemove;
 			NewLODInfo.BonesToPrioritize = LODInfoArray[LastIndex].BonesToPrioritize;
+			NewLODInfo.SectionsToPrioritize = LODInfoArray[LastIndex].SectionsToPrioritize;
 			// now find reduction setting
 			for (int32 SubLOD = LastIndex; SubLOD >= 0; --SubLOD)
 			{
@@ -5541,6 +5576,7 @@ FGuid FSkeletalMeshLODInfo::ComputeDeriveDataCacheKey(const FSkeletalMeshLODGrou
 
 	Ar << BonesToRemove;
 	Ar << BonesToPrioritize;
+	Ar << SectionsToPrioritize;
 	Ar << WeightOfPrioritization;
 
 	//TODO: Ask the derivedata key of the UObject reference by FSoftObjectPath. So if someone change the UObject, this LODs will get dirty
diff --git a/Engine/Source/Runtime/Engine/Private/SkeletalMeshLODSettings.cpp b/Engine/Source/Runtime/Engine/Private/SkeletalMeshLODSettings.cpp
--- a/Engine/Source/Runtime/Engine/Private/SkeletalMeshLODSettings.cpp
+++ b/Engine/Source/Runtime/Engine/Private/SkeletalMeshLODSettings.cpp
@@ -7,6 +7,8 @@
 #include "Animation/AnimSequence.h"
 #include "UObject/FortniteMainBranchObjectVersion.h"
 #include "UObject/EditorObjectVersion.h"
+#include "Rendering/SkeletalMeshLODModel.h"
+#include "Rendering/SkeletalMeshModel.h"
 
 DEFINE_LOG_CATEGORY_STATIC(LogSkeletalMeshLODSettings, Warning, All)
 
@@ -79,6 +81,33 @@ bool USkeletalMeshLODSettings::SetLODSettingsToMesh(USkeletalMesh* InMesh, int32
 			}
 		}
 
+#if WITH_EDITOR
+		// copy the shared setting to mesh setting
+		// make sure we have the section in the mesh 
+		const FSkeletalMeshLODModel& LodModel = InMesh->GetImportedModel()->LODModels[LODIndex];
+#endif
+		LODInfo->SectionsToPrioritize.Reset();
+		for (const int32 SettingSectionIndex : Setting.SectionsToPrioritize)
+		{
+			int32 FinalSettingSectionIndex = SettingSectionIndex;
+#if WITH_EDITOR
+			FinalSettingSectionIndex = INDEX_NONE;
+			for (int32 SectionIndex = 0; SectionIndex < LodModel.Sections.Num(); ++SectionIndex)
+			{
+				const FSkelMeshSection& Section = LodModel.Sections[SectionIndex];
+				if (Section.ChunkedParentSectionIndex == INDEX_NONE && SettingSectionIndex == Section.OriginalDataSectionIndex)
+				{
+					FinalSettingSectionIndex = Section.OriginalDataSectionIndex;
+					break;
+				}
+			}
+#endif
+			if (FinalSettingSectionIndex != INDEX_NONE)
+			{
+				LODInfo->SectionsToPrioritize.AddUnique(FSectionReference(FinalSettingSectionIndex));
+			}
+		}
+
 		if (Setting.BoneFilterActionOption == EBoneFilterActionOption::Remove)
 		{
 			for (int32 BoneListIndex = 0; BoneListIndex < Setting.BoneList.Num(); ++BoneListIndex)
@@ -218,6 +247,28 @@ int32 USkeletalMeshLODSettings::SetLODSettingsFromMesh(USkeletalMesh* InMesh)
 				Setting.BonesToPrioritize.Add(Bone.BoneName);
 			}
 
+#if WITH_EDITOR
+			// copy the shared setting to mesh setting
+			// make sure we have the section in the mesh 
+			const FSkeletalMeshLODModel& LodModel = InMesh->GetImportedModel()->LODModels[Index];
+#endif
+			Setting.SectionsToPrioritize.Reset();
+			for (const FSectionReference& SectionReference : LODInfo->SectionsToPrioritize)
+			{
+				int32 FinalSettingSectionIndex = SectionReference.SectionIndex;
+#if WITH_EDITOR
+				//In editor we can validate the prioritized sections
+				if (!SectionReference.IsValidToEvaluate(LodModel))
+				{
+					FinalSettingSectionIndex = INDEX_NONE;
+				}
+#endif
+				if (FinalSettingSectionIndex != INDEX_NONE)
+				{
+					Setting.SectionsToPrioritize.AddUnique(FinalSettingSectionIndex);
+				}
+			}
+
 			Setting.WeightOfPrioritization = LODInfo->WeightOfPrioritization;
 			Setting.BakePose = LODInfo->BakePose;
 			Setting.BoneFilterActionOption = EBoneFilterActionOption::Remove;
