Fix reduction assert during cooking
#rnx

diff --git a/Engine/Plugins/Experimental/SkeletalReduction/Source/Private/SkeletalMeshReductionPlugin.cpp b/Engine/Plugins/Experimental/SkeletalReduction/Source/Private/SkeletalMeshReductionPlugin.cpp
--- a/Engine/Plugins/Experimental/SkeletalReduction/Source/Private/SkeletalMeshReductionPlugin.cpp
+++ b/Engine/Plugins/Experimental/SkeletalReduction/Source/Private/SkeletalMeshReductionPlugin.cpp
@@ -812,7 +812,7 @@ void FQuadricSkeletalMeshReduction::ConvertToFSkinnedSkeletalMesh( const FSkelet
 	{
 		const int32 InvalidSectionId = -1;
 		VertToSectionMap->Empty();
-		VertToSectionMap->Init(InvalidSectionId, VertexCount);
+		VertToSectionMap->Reserve(VertexCount);
 	
 		TArray<int32>& ToSection = *VertToSectionMap;
 		for (int32 s = 0; s < SectionCount; ++s)
@@ -824,6 +824,7 @@ void FQuadricSkeletalMeshReduction::ConvertToFSkinnedSkeletalMesh( const FSkelet
 			int32 OriginalDataSectionIndex = SrcLODModel.Sections[s].OriginalDataSectionIndex;
 
 			const FSectionRange VertexRange = SectionRangeArray[s];
+			VertToSectionMap->AddZeroed(SrcLODModel.Sections[s].NumVertices);
 
 			for (int32 v = VertexRange.Begin; v < VertexRange.End; ++v)
 			{
