[Backout] - CL18299908
Original CL Desc
-----------------------------------------------------------------
Add section to prioritize sharing the prioritize weight value. This will allow to reduce less a particular section.

diff --git a/Engine/Plugins/Experimental/SkeletalReduction/Source/Private/SkeletalMeshReductionPlugin.cpp b/Engine/Plugins/Experimental/SkeletalReduction/Source/Private/SkeletalMeshReductionPlugin.cpp
--- a/Engine/Plugins/Experimental/SkeletalReduction/Source/Private/SkeletalMeshReductionPlugin.cpp
+++ b/Engine/Plugins/Experimental/SkeletalReduction/Source/Private/SkeletalMeshReductionPlugin.cpp
@@ -158,9 +158,9 @@ private:
 	struct  FSectionRange;
 
 	/**
-	* Important bones and sections when simplifying
+	* Important bones when simplifying
 	*/
-	struct FPrioritizedFeatures;
+	struct FImportantBones;
 
 private:
 
@@ -191,7 +191,7 @@ private:
 		                         const FBoxSphereBounds& Bounds,
 		                         const FReferenceSkeleton& RefSkeleton,
 		                         FSkeletalMeshOptimizationSettings Settings,
-								 const FPrioritizedFeatures& PrioritizedFeatures,
+								 const FImportantBones& ImportantBones,
 		                         const TArray<FMatrix>& BoneMatrices,
 		                         const int32 LODIndex,
 								 const bool bReducingSourceModel,
@@ -211,13 +211,11 @@ private:
 	* @param BoneMatrices   - define the configuration of the model
 	* @param LODIndex       - target index for the result.
 	* @param OutSkinnedMesh - the posed mesh 
-	* @param VerToSectionMap - if provided, on return this array will allow VertexID to SectionID lookup
 	*/
 	void ConvertToFSkinnedSkeletalMesh( const FSkeletalMeshLODModel& SrcLODModel,
 		                                const TArray<FMatrix>& BoneMatrices,
 		                                const int32 LODIndex,
-		                                SkeletalSimplifier::FSkinnedSkeletalMesh& OutSkinnedMesh,
-										TArray<int32>* VertToSectionMap = nullptr) const;
+		                                SkeletalSimplifier::FSkinnedSkeletalMesh& OutSkinnedMesh ) const;
 
 	/**
 	* Generate a SkeletalMeshLODModel from a SkinnedSkeletalMesh and ReferenceSkeleton
@@ -300,11 +298,9 @@ private:
 	void TrimBonesPerVert( SkeletalSimplifier::FSkinnedSkeletalMesh& Mesh, int32 MaxBonesPerVert ) const ;
 
 	/**
-	* If a vertex has one of the prioritized bones as its' major bone, associated the PrioritizedFeatures.Weight
-	* Likewise if the section ID for the vert is in a prioritized section, associate the prioritized weight 
+	* If a vertex has one of the important bones as its' major bone, associated the ImportantBones.Weight
 	*/
-	void UpdateSpecializedVertWeights(const FPrioritizedFeatures& PrioritizedFeatures, const TArray<int32>& VertToSectionId, SkeletalSimplifier::FSkinnedSkeletalMesh& SkinnedSkeletalMesh) const;
-
+	void UpdateSpecializedVertWeights(const FImportantBones& ImportantBones, SkeletalSimplifier::FSkinnedSkeletalMesh& SkinnedSkeletalMesh) const;
 
 };
 
@@ -324,16 +320,10 @@ struct FQuadricSkeletalMeshReduction::FSectionRange
 	int32 End;
 };
 
-struct FQuadricSkeletalMeshReduction::FPrioritizedFeatures
+struct FQuadricSkeletalMeshReduction::FImportantBones
 {
-	TArray<int32> BoneIds;
-	TArray<int32> SectionIds;
+	TArray<int32> Ids;
 	float Weight;
-
-	bool HasFeatures() const
-	{
-		return (BoneIds.Num() >0 || SectionIds.Num() > 0);
-	}
 };
 /**
 *  Required MeshReduction Interface.
@@ -487,8 +477,7 @@ bool FQuadricSkeletalMeshReduction::RemoveMeshSection(FSkeletalMeshLODModel& Mod
 void FQuadricSkeletalMeshReduction::ConvertToFSkinnedSkeletalMesh( const FSkeletalMeshLODModel& SrcLODModel,
 	                                                               const TArray<FMatrix>& BoneMatrices,
 	                                                               const int32 LODIndex,
-	                                                               SkeletalSimplifier::FSkinnedSkeletalMesh& OutSkinnedMesh,
-																   TArray<int32>* VertToSectionMap) const
+	                                                               SkeletalSimplifier::FSkinnedSkeletalMesh& OutSkinnedMesh) const
 {
 
 
@@ -807,30 +796,6 @@ void FQuadricSkeletalMeshReduction::ConvertToFSkinnedSkeletalMesh( const FSkelet
 		}
 	}
 
-	// if requested, make a map of the vertex to original section id
-	if (VertToSectionMap != nullptr)
-	{
-		const int32 InvalidSectionId = -1;
-		VertToSectionMap->Empty();
-		VertToSectionMap->Init(InvalidSectionId, VertexCount);
-	
-		TArray<int32>& ToSection = *VertToSectionMap;
-		for (int32 s = 0; s < SectionCount; ++s)
-		{
-			if (SkipSection(s))
-			{
-				continue;
-			}
-			int32 OriginalDataSectionIndex = SrcLODModel.Sections[s].OriginalDataSectionIndex;
-
-			const FSectionRange VertexRange = SectionRangeArray[s];
-
-			for (int32 v = VertexRange.Begin; v < VertexRange.End; ++v)
-			{
-				ToSection[v] = OriginalDataSectionIndex;
-			}
-		}
-	}
 	// Put the vertex in a "correct" state.
 	//    "corrects" normals (ensures that they are orthonormal)
 	//    re-orders the bones by weight (highest to lowest)
@@ -857,28 +822,16 @@ void FQuadricSkeletalMeshReduction::ConvertToFSkinnedSkeletalMesh( const FSkelet
 }
 
 
-void FQuadricSkeletalMeshReduction::UpdateSpecializedVertWeights(const FPrioritizedFeatures& PrioritizedFeatures, const TArray<int32>& VertToOriginalSectionId, SkeletalSimplifier::FSkinnedSkeletalMesh& SkinnedSkeletalMesh) const
+void FQuadricSkeletalMeshReduction::UpdateSpecializedVertWeights(const FImportantBones& ImportantBones, SkeletalSimplifier::FSkinnedSkeletalMesh& SkinnedSkeletalMesh) const
 {
-
-	if (!PrioritizedFeatures.HasFeatures())
-	{
-		return;
-	}
-	
-	const float Weight = PrioritizedFeatures.Weight;
+	const float Weight = ImportantBones.Weight;
 
 	int32 NumVerts = SkinnedSkeletalMesh.NumVertices();
-	check(VertToOriginalSectionId.Num() == NumVerts);
 
-	const int32 InvalidSectionId = -1;
 	//If a vertex has one of the important bones as its' major bone, associated the ImportantBones.Weight
 	for (int32 i = 0; i < NumVerts; ++i)
 	{
-		const int32 OriginalSectionId = VertToOriginalSectionId[i];
 		auto& Vert = SkinnedSkeletalMesh.VertexBuffer[i];
-		Vert.SpecializedWeight = 0.f;
-
-		// Test if this vert is associated with a priority bone
 		const auto& Bones = Vert.GetSparseBones();
 		if (!Bones.bIsEmpty())
 		{
@@ -886,19 +839,15 @@ void FQuadricSkeletalMeshReduction::UpdateSpecializedVertWeights(const FPrioriti
 
 			const int32 FirstBone = CIter.Key(); // Bones are ordered by descending weight
 
-			if (PrioritizedFeatures.BoneIds.Contains(FirstBone))
+			if (ImportantBones.Ids.Contains(FirstBone))
 			{
 				Vert.SpecializedWeight = Weight;
 			}
 		}
-		
-		// Test if this vert is associated with a priority section
-		if (OriginalSectionId != InvalidSectionId && PrioritizedFeatures.SectionIds.Contains(OriginalSectionId))
+		else
 		{
-			//Using += here will make bone and section prioritized verts stronger than only bone or only section prioritized verts
-			Vert.SpecializedWeight += Weight;
+			Vert.SpecializedWeight = 0.f;
 		}
-
 	}
 }
 
@@ -1725,7 +1674,7 @@ bool FQuadricSkeletalMeshReduction::ReduceSkeletalLODModel( const FSkeletalMeshL
 	                                                        const FBoxSphereBounds& Bounds,
 	                                                        const FReferenceSkeleton& RefSkeleton,
 	                                                        FSkeletalMeshOptimizationSettings Settings,
-															const FPrioritizedFeatures& PrioritizedFeatures,
+															const FImportantBones& ImportantBones,
 	                                                        const TArray<FMatrix>& BoneMatrices,
 	                                                        const int32 LODIndex,
 															const bool bReducingSourceModel,
@@ -1763,9 +1712,9 @@ bool FQuadricSkeletalMeshReduction::ReduceSkeletalLODModel( const FSkeletalMeshL
 	}
 	
 	// Generate a single skinned mesh form the SrcModel.  This mesh has per-vertex tangent space.
-	TArray<int32> VertToOriginalSectionId; 
+
 	SkeletalSimplifier::FSkinnedSkeletalMesh SkinnedSkeletalMesh;
-	ConvertToFSkinnedSkeletalMesh(SrcModel, BoneMatrices, LODIndex, SkinnedSkeletalMesh, &VertToOriginalSectionId);
+	ConvertToFSkinnedSkeletalMesh(SrcModel, BoneMatrices, LODIndex, SkinnedSkeletalMesh);
 
 	int32 IterationNum = 0;
 	//We keep the original MaxNumVerts because if we iterate we want to still compare with the original request.
@@ -1774,10 +1723,11 @@ bool FQuadricSkeletalMeshReduction::ReduceSkeletalLODModel( const FSkeletalMeshL
 	{
 		if (bOptimizeMesh)
 		{
-
-			// Add specialized weights for verts associated with "important" bones or sections
-			UpdateSpecializedVertWeights(PrioritizedFeatures, VertToOriginalSectionId, SkinnedSkeletalMesh);
-
+			if (ImportantBones.Ids.Num() > 0)
+			{
+				// Add specialized weights for verts associated with "important" bones.
+				UpdateSpecializedVertWeights(ImportantBones, SkinnedSkeletalMesh);
+			}
 
 			// Capture the UV bounds from the source mesh.
 
@@ -1900,28 +1850,18 @@ void FQuadricSkeletalMeshReduction::ReduceSkeletalMesh(USkeletalMesh& SkeletalMe
 	// Struct to identify important bones.  Vertices associated with these bones
 	// will have additional collapse weight added to them.
 
-	FPrioritizedFeatures  PrioritizedFeatures;
+	FImportantBones  ImportantBones;
 	{
-		
-		const float PrioritizationWeight = LODInfo->WeightOfPrioritization;
-		PrioritizedFeatures.Weight = PrioritizationWeight;
-
 		const TArray<FBoneReference>& BonesToPrioritize = LODInfo->BonesToPrioritize;
+		const float BonePrioritizationWeight = LODInfo->WeightOfPrioritization;
+
+		ImportantBones.Weight = BonePrioritizationWeight;
 		for (const FBoneReference& BoneReference : BonesToPrioritize)
 		{
 			int32 BoneId = SkeletalMesh.GetRefSkeleton().FindRawBoneIndex(BoneReference.BoneName);
 
 			// Q: should we exclude BoneId = 0?
-			PrioritizedFeatures.BoneIds.AddUnique(BoneId);
-		}
-
-		const TArray<FSectionReference>& SectionsToPrioritize = LODInfo->SectionsToPrioritize;
-		for (const FSectionReference& SectionReference : SectionsToPrioritize)
-		{
-			if (SectionReference.IsValidToEvaluate(SkeletalMeshResource.LODModels[LODIndex]))
-			{
-				PrioritizedFeatures.SectionIds.AddUnique(SectionReference.SectionIndex);
-			}
+			ImportantBones.Ids.AddUnique(BoneId);
 		}
 	}
 	// select which mesh we're reducing from
@@ -2213,7 +2153,7 @@ void FQuadricSkeletalMeshReduction::ReduceSkeletalMesh(USkeletalMesh& SkeletalMe
 	}
 
 	// Reduce LOD model with SrcMesh if src mesh has more then 1 triangle
-	if (SrcModel->NumVertices > 3 && ReduceSkeletalLODModel(*SrcModel, *NewModel, SkeletalMesh.GetPathName(), SkeletalMesh.GetImportedBounds(), SkeletalMesh.GetRefSkeleton(), Settings, PrioritizedFeatures, RelativeToRefPoseMatrices, LODIndex, bReducingSourceModel, TargetPlatform))
+	if (SrcModel->NumVertices > 3 && ReduceSkeletalLODModel(*SrcModel, *NewModel, SkeletalMesh.GetPathName(), SkeletalMesh.GetImportedBounds(), SkeletalMesh.GetRefSkeleton(), Settings, ImportantBones, RelativeToRefPoseMatrices, LODIndex, bReducingSourceModel, TargetPlatform))
 	{
 		FSkeletalMeshLODInfo* ReducedLODInfoPtr = SkeletalMesh.GetLODInfo(LODIndex);
 		check(ReducedLODInfoPtr);
diff --git a/Engine/Source/Editor/Persona/Private/PersonaMeshDetails.cpp b/Engine/Source/Editor/Persona/Private/PersonaMeshDetails.cpp
--- a/Engine/Source/Editor/Persona/Private/PersonaMeshDetails.cpp
+++ b/Engine/Source/Editor/Persona/Private/PersonaMeshDetails.cpp
@@ -2477,12 +2477,6 @@ void FPersonaMeshDetails::CustomizeLODInfoSetingsDetails(IDetailLayoutBuilder& D
 	PreventAttributePostEditChange(BonesToPrioritizeHandle, LODIndex, BonesToPrioritizeHandle->GetProperty()->GetFName(), bForceComponentRefreshFalse);
 	DetailLayout.HideProperty(BonesToPrioritizeHandle);
 
-	TSharedPtr<IPropertyHandle> SectionsToPrioritizeHandle = LODInfoChild->GetChildHandle(GET_MEMBER_NAME_CHECKED(FSkeletalMeshLODInfo, SectionsToPrioritize));
-	IDetailPropertyRow& SectionsToPrioritizeRow = LODInfoGroup.AddPropertyRow(SectionsToPrioritizeHandle->AsShared());
-	SectionsToPrioritizeRow.IsEnabled(EnabledAttrib);
-	PreventAttributePostEditChange(SectionsToPrioritizeHandle, LODIndex, SectionsToPrioritizeHandle->GetProperty()->GetFName(), bForceComponentRefreshFalse);
-	DetailLayout.HideProperty(SectionsToPrioritizeHandle);
-
 	TSharedPtr<IPropertyHandle> WeightToPriortizeHandle = LODInfoChild->GetChildHandle(GET_MEMBER_NAME_CHECKED(FSkeletalMeshLODInfo, WeightOfPrioritization));
 	IDetailPropertyRow& WeightToPriortizeRow = LODInfoGroup.AddPropertyRow(WeightToPriortizeHandle->AsShared());
 	WeightToPriortizeRow.IsEnabled(EnabledAttrib);
@@ -2490,8 +2484,7 @@ void FPersonaMeshDetails::CustomizeLODInfoSetingsDetails(IDetailLayoutBuilder& D
 	DetailLayout.HideProperty(WeightToPriortizeHandle);
 
 	const TArray<FName> HiddenProperties = { GET_MEMBER_NAME_CHECKED(FSkeletalMeshLODInfo, ReductionSettings), GET_MEMBER_NAME_CHECKED(FSkeletalMeshLODInfo, BakePose), GET_MEMBER_NAME_CHECKED(FSkeletalMeshLODInfo, BakePoseOverride), GET_MEMBER_NAME_CHECKED(FSkeletalMeshLODInfo, BonesToRemove),
-		GET_MEMBER_NAME_CHECKED(FSkeletalMeshLODInfo, BonesToPrioritize), GET_MEMBER_NAME_CHECKED(FSkeletalMeshLODInfo, SectionsToPrioritize), GET_MEMBER_NAME_CHECKED(FSkeletalMeshLODInfo, WeightOfPrioritization), GET_MEMBER_NAME_CHECKED(FSkeletalMeshLODInfo, ScreenSize),
-		GET_MEMBER_NAME_CHECKED(FSkeletalMeshLODInfo, LODHysteresis), GET_MEMBER_NAME_CHECKED(FSkeletalMeshLODInfo, BuildSettings) };
+		GET_MEMBER_NAME_CHECKED(FSkeletalMeshLODInfo, BonesToPrioritize), GET_MEMBER_NAME_CHECKED(FSkeletalMeshLODInfo, WeightOfPrioritization), GET_MEMBER_NAME_CHECKED(FSkeletalMeshLODInfo, ScreenSize), GET_MEMBER_NAME_CHECKED(FSkeletalMeshLODInfo, LODHysteresis), GET_MEMBER_NAME_CHECKED(FSkeletalMeshLODInfo, BuildSettings) };
 	for (uint32 ChildIndex = 0; ChildIndex < NumInfoChildren; ++ChildIndex)
 	{
 		TSharedRef<IPropertyHandle> LODInfoChildHandle = LODInfoChild->GetChildHandle(ChildIndex).ToSharedRef();
diff --git a/Engine/Source/Editor/Persona/Private/PersonaModule.cpp b/Engine/Source/Editor/Persona/Private/PersonaModule.cpp
--- a/Engine/Source/Editor/Persona/Private/PersonaModule.cpp
+++ b/Engine/Source/Editor/Persona/Private/PersonaModule.cpp
@@ -92,7 +92,6 @@
 #include "AnimSequenceTimelineCommands.h"
 #include "SAnimMontageSectionsPanel.h"
 #include "Animation/AnimSequenceHelpers.h"
-#include "SkeletalMeshReferenceSectionDetails.h"
 
 IMPLEMENT_MODULE( FPersonaModule, Persona );
 
@@ -134,8 +133,6 @@ void FPersonaModule::StartupModule()
 
 		PropertyModule.RegisterCustomPropertyTypeLayout("SkeletalMeshSamplingRegionBoneFilter", FOnGetPropertyTypeCustomizationInstance::CreateStatic(&FNiagaraSkeletalMeshRegionBoneFilterDetails::MakeInstance));
 		PropertyModule.RegisterCustomPropertyTypeLayout("SkeletalMeshSamplingRegionMaterialFilter", FOnGetPropertyTypeCustomizationInstance::CreateStatic(&FNiagaraSkeletalMeshRegionMaterialFilterDetails::MakeInstance));
-
-		PropertyModule.RegisterCustomPropertyTypeLayout("SectionReference", FOnGetPropertyTypeCustomizationInstance::CreateStatic(&FSectionReferenceCustomization::MakeInstance));
 	}
 
 	// Register the editor modes
diff --git a/Engine/Source/Editor/Persona/Private/Shared/ReferenceSectionSelectionWidget.cpp b/Engine/Source/Editor/Persona/Private/Shared/ReferenceSectionSelectionWidget.cpp
deleted file mode 100644
diff --git a/Engine/Source/Editor/Persona/Private/SkeletalMeshReferenceSectionDetails.cpp b/Engine/Source/Editor/Persona/Private/SkeletalMeshReferenceSectionDetails.cpp
deleted file mode 100644
diff --git a/Engine/Source/Editor/Persona/Private/SkeletalMeshReferenceSectionDetails.h b/Engine/Source/Editor/Persona/Private/SkeletalMeshReferenceSectionDetails.h
deleted file mode 100644
diff --git a/Engine/Source/Editor/Persona/Public/ReferenceSectionSelectionWidget.h b/Engine/Source/Editor/Persona/Public/ReferenceSectionSelectionWidget.h
deleted file mode 100644
diff --git a/Engine/Source/Editor/UnrealEd/Private/ImportUtils/SkeletalMeshImportUtils.cpp b/Engine/Source/Editor/UnrealEd/Private/ImportUtils/SkeletalMeshImportUtils.cpp
--- a/Engine/Source/Editor/UnrealEd/Private/ImportUtils/SkeletalMeshImportUtils.cpp
+++ b/Engine/Source/Editor/UnrealEd/Private/ImportUtils/SkeletalMeshImportUtils.cpp
@@ -461,7 +461,6 @@ void SkeletalMesUtilsImpl::RestoreLODInfo(const TSharedPtr<const FExistingSkelMe
 	}
 	ImportedLODInfo.BonesToRemove = ExistingLODInfo.BonesToRemove;
 	ImportedLODInfo.BonesToPrioritize = ExistingLODInfo.BonesToPrioritize;
-	ImportedLODInfo.SectionsToPrioritize = ExistingLODInfo.SectionsToPrioritize;
 	ImportedLODInfo.WeightOfPrioritization = ExistingLODInfo.WeightOfPrioritization;
 	ImportedLODInfo.BakePose = ExistingLODInfo.BakePose;
 	ImportedLODInfo.BakePoseOverride = ExistingLODInfo.BakePoseOverride;
diff --git a/Engine/Source/Runtime/Engine/Classes/Engine/SkeletalMesh.h b/Engine/Source/Runtime/Engine/Classes/Engine/SkeletalMesh.h
--- a/Engine/Source/Runtime/Engine/Classes/Engine/SkeletalMesh.h
+++ b/Engine/Source/Runtime/Engine/Classes/Engine/SkeletalMesh.h
@@ -359,51 +359,6 @@ struct ENGINE_API FSkeletalMeshClothBuildParams
 	TSoftObjectPtr<UPhysicsAsset> PhysicsAsset;
 };
 
-USTRUCT()
-struct FSectionReference
-{
-	GENERATED_USTRUCT_BODY()
-
-	/** Index of the section we reference. **/
-	UPROPERTY(EditAnywhere, Category = SectionReference)
-	int32 SectionIndex;
-
-	FSectionReference()
-		: SectionIndex(INDEX_NONE)
-	{
-	}
-
-	FSectionReference(const int32& InSectionIndex)
-		: SectionIndex(InSectionIndex)
-	{
-	}
-
-	bool operator==(const FSectionReference& Other) const
-	{
-		return SectionIndex == Other.SectionIndex;
-	}
-
-#if WITH_EDITOR
-	/** return true if it has a valid section index for LodModel parameter **/
-	ENGINE_API bool IsValidToEvaluate(const FSkeletalMeshLODModel& LodModel) const;
-
-	const struct FSkelMeshSection* GetMeshLodSection(const FSkeletalMeshLODModel& LodModel) const;
-	int32 GetMeshLodSectionIndex(const FSkeletalMeshLODModel& LodModel) const;
-#endif
-
-	bool Serialize(FArchive& Ar)
-	{
-		Ar << SectionIndex;
-		return true;
-	}
-
-	friend FArchive& operator<<(FArchive& Ar, FSectionReference& B)
-	{
-		B.Serialize(Ar);
-		return Ar;
-	}
-};
-
 /** Struct containing information for a particular LOD level, such as materials and info for when to use it. */
 USTRUCT()
 struct FSkeletalMeshLODInfo
@@ -451,15 +406,11 @@ struct FSkeletalMeshLODInfo
 	UPROPERTY(EditAnywhere, Category = ReductionSettings)
 	TArray<FBoneReference> BonesToRemove;
 
-	/** Bones which should be prioritized for the quality, this will be weighted toward keeping source data. Use WeightOfPrioritization to control the value. */
+	/** Bones which should be prioritized for the quality, this will be weighted toward keeping source data. */
 	UPROPERTY(EditAnywhere, Category = ReductionSettings)
 	TArray<FBoneReference> BonesToPrioritize;
 
-	/** Sections which should be prioritized for the quality, this will be weighted toward keeping source data. Use WeightOfPrioritization to control the value. */
-	UPROPERTY(EditAnywhere, Category = ReductionSettings)
-	TArray<FSectionReference> SectionsToPrioritize;
-
-	/** How much to consideration to give BonesToPrioritize and SectionsToPrioritize.  The weight is an additional vertex simplification penalty where 0 means nothing. */
+	/** How much to consideration to give BonesToPrioritize.  The weight is an additional vertex simplification penalty where 0 means nothing. */
 	UPROPERTY(EditAnywhere, Category = ReductionSettings, meta = (UIMin = "0.0", ClampMin = "0.0"))
 	float WeightOfPrioritization;
 
diff --git a/Engine/Source/Runtime/Engine/Classes/Engine/SkeletalMeshLODSettings.h b/Engine/Source/Runtime/Engine/Classes/Engine/SkeletalMeshLODSettings.h
--- a/Engine/Source/Runtime/Engine/Classes/Engine/SkeletalMeshLODSettings.h
+++ b/Engine/Source/Runtime/Engine/Classes/Engine/SkeletalMeshLODSettings.h
@@ -93,16 +93,12 @@ struct FSkeletalMeshLODGroupSettings
 	UPROPERTY(EditAnywhere, Category = Reduction)
 	TArray<FBoneFilter> BoneList;
 
-	/** Bones which should be prioritized for the quality, this will be weighted toward keeping source data. Use WeightOfPrioritization to control the value. */
+	/** Bones which should be prioritized for the quality, this will be weighted toward keeping source data. */
 	UPROPERTY(EditAnywhere, Category = Reduction)
 	TArray<FName> BonesToPrioritize;
 
-	/** Sections which should be prioritized for the quality, this will be weighted toward keeping source data. Use WeightOfPrioritization to control the value. */
-	UPROPERTY(EditAnywhere, Category = Reduction)
-	TArray<int32> SectionsToPrioritize;
-
-	/** How much to consideration to give BonesToPrioritize and SectionsToPrioritize.  The weight is an additional vertex simplification penalty where 0 means nothing. */
-	UPROPERTY(EditAnywhere, Category = Reduction, meta = (UIMin = "0.0", ClampMin = "0.0"))
+	/** Weight of how much consider for BonesToPrioritize. 0 means nothing, and 1 means take all source */
+	UPROPERTY(EditAnywhere, Category = Reduction, meta = (UIMin = "0.0", UIMax = "1.0", ClampMin = "0.0", ClampMax = "1.0"))
 	float WeightOfPrioritization;
 
 	/** Pose which should be used to reskin vertex influences for which the bones will be removed in this LOD level, uses ref-pose by default */
diff --git a/Engine/Source/Runtime/Engine/Private/SkeletalMesh.cpp b/Engine/Source/Runtime/Engine/Private/SkeletalMesh.cpp
--- a/Engine/Source/Runtime/Engine/Private/SkeletalMesh.cpp
+++ b/Engine/Source/Runtime/Engine/Private/SkeletalMesh.cpp
@@ -440,47 +440,13 @@ void FScopedSkeletalMeshPostEditChange::SetSkeletalMesh(USkeletalMesh* InSkeleta
 	}
 }
 
-
-/*-----------------------------------------------------------------------------
-	FSectionReference
------------------------------------------------------------------------------*/
-
-bool FSectionReference::IsValidToEvaluate(const FSkeletalMeshLODModel& LodModel) const
-{
-	return GetMeshLodSection(LodModel) != nullptr;
-}
-
-const FSkelMeshSection* FSectionReference::GetMeshLodSection(const FSkeletalMeshLODModel& LodModel) const
-{
-	for (int32 LodModelSectionIndex = 0; LodModelSectionIndex < LodModel.Sections.Num(); ++LodModelSectionIndex)
-	{
-		const FSkelMeshSection& Section = LodModel.Sections[LodModelSectionIndex];
-		if (Section.ChunkedParentSectionIndex == INDEX_NONE && SectionIndex == Section.OriginalDataSectionIndex)
-		{
-			return &Section;
-		}
-	}
-	return nullptr;
-}
-
-int32 FSectionReference::GetMeshLodSectionIndex(const FSkeletalMeshLODModel& LodModel) const
-{
-	for (int32 LodModelSectionIndex = 0; LodModelSectionIndex < LodModel.Sections.Num(); ++LodModelSectionIndex)
-	{
-		const FSkelMeshSection& Section = LodModel.Sections[LodModelSectionIndex];
-		if (Section.ChunkedParentSectionIndex == INDEX_NONE && SectionIndex == Section.OriginalDataSectionIndex)
-		{
-			return LodModelSectionIndex;
-		}
-	}
-	return INDEX_NONE;
-}
-
+#endif //WITH_EDITOR
 
 /*-----------------------------------------------------------------------------
 	FSkeletalMeshAsyncBuildWorker
 -----------------------------------------------------------------------------*/
 
+#if WITH_EDITOR
 
 void FSkeletalMeshAsyncBuildWorker::DoWork()
 {
@@ -5317,7 +5283,6 @@ FSkeletalMeshLODInfo& USkeletalMesh::AddLODInfo()
 			NewLODInfo.BakePoseOverride = LODInfoArray[LastIndex].BakePoseOverride;
 			NewLODInfo.BonesToRemove = LODInfoArray[LastIndex].BonesToRemove;
 			NewLODInfo.BonesToPrioritize = LODInfoArray[LastIndex].BonesToPrioritize;
-			NewLODInfo.SectionsToPrioritize = LODInfoArray[LastIndex].SectionsToPrioritize;
 			// now find reduction setting
 			for (int32 SubLOD = LastIndex; SubLOD >= 0; --SubLOD)
 			{
@@ -5576,7 +5541,6 @@ FGuid FSkeletalMeshLODInfo::ComputeDeriveDataCacheKey(const FSkeletalMeshLODGrou
 
 	Ar << BonesToRemove;
 	Ar << BonesToPrioritize;
-	Ar << SectionsToPrioritize;
 	Ar << WeightOfPrioritization;
 
 	//TODO: Ask the derivedata key of the UObject reference by FSoftObjectPath. So if someone change the UObject, this LODs will get dirty
diff --git a/Engine/Source/Runtime/Engine/Private/SkeletalMeshLODSettings.cpp b/Engine/Source/Runtime/Engine/Private/SkeletalMeshLODSettings.cpp
--- a/Engine/Source/Runtime/Engine/Private/SkeletalMeshLODSettings.cpp
+++ b/Engine/Source/Runtime/Engine/Private/SkeletalMeshLODSettings.cpp
@@ -7,8 +7,6 @@
 #include "Animation/AnimSequence.h"
 #include "UObject/FortniteMainBranchObjectVersion.h"
 #include "UObject/EditorObjectVersion.h"
-#include "Rendering/SkeletalMeshLODModel.h"
-#include "Rendering/SkeletalMeshModel.h"
 
 DEFINE_LOG_CATEGORY_STATIC(LogSkeletalMeshLODSettings, Warning, All)
 
@@ -81,33 +79,6 @@ bool USkeletalMeshLODSettings::SetLODSettingsToMesh(USkeletalMesh* InMesh, int32
 			}
 		}
 
-#if WITH_EDITOR
-		// copy the shared setting to mesh setting
-		// make sure we have the section in the mesh 
-		const FSkeletalMeshLODModel& LodModel = InMesh->GetImportedModel()->LODModels[LODIndex];
-#endif
-		LODInfo->SectionsToPrioritize.Reset();
-		for (const int32 SettingSectionIndex : Setting.SectionsToPrioritize)
-		{
-			int32 FinalSettingSectionIndex = SettingSectionIndex;
-#if WITH_EDITOR
-			FinalSettingSectionIndex = INDEX_NONE;
-			for (int32 SectionIndex = 0; SectionIndex < LodModel.Sections.Num(); ++SectionIndex)
-			{
-				const FSkelMeshSection& Section = LodModel.Sections[SectionIndex];
-				if (Section.ChunkedParentSectionIndex == INDEX_NONE && SettingSectionIndex == Section.OriginalDataSectionIndex)
-				{
-					FinalSettingSectionIndex = Section.OriginalDataSectionIndex;
-					break;
-				}
-			}
-#endif
-			if (FinalSettingSectionIndex != INDEX_NONE)
-			{
-				LODInfo->SectionsToPrioritize.AddUnique(FSectionReference(FinalSettingSectionIndex));
-			}
-		}
-
 		if (Setting.BoneFilterActionOption == EBoneFilterActionOption::Remove)
 		{
 			for (int32 BoneListIndex = 0; BoneListIndex < Setting.BoneList.Num(); ++BoneListIndex)
@@ -247,28 +218,6 @@ int32 USkeletalMeshLODSettings::SetLODSettingsFromMesh(USkeletalMesh* InMesh)
 				Setting.BonesToPrioritize.Add(Bone.BoneName);
 			}
 
-#if WITH_EDITOR
-			// copy the shared setting to mesh setting
-			// make sure we have the section in the mesh 
-			const FSkeletalMeshLODModel& LodModel = InMesh->GetImportedModel()->LODModels[Index];
-#endif
-			Setting.SectionsToPrioritize.Reset();
-			for (const FSectionReference& SectionReference : LODInfo->SectionsToPrioritize)
-			{
-				int32 FinalSettingSectionIndex = SectionReference.SectionIndex;
-#if WITH_EDITOR
-				//In editor we can validate the prioritized sections
-				if (!SectionReference.IsValidToEvaluate(LodModel))
-				{
-					FinalSettingSectionIndex = INDEX_NONE;
-				}
-#endif
-				if (FinalSettingSectionIndex != INDEX_NONE)
-				{
-					Setting.SectionsToPrioritize.AddUnique(FinalSettingSectionIndex);
-				}
-			}
-
 			Setting.WeightOfPrioritization = LODInfo->WeightOfPrioritization;
 			Setting.BakePose = LODInfo->BakePose;
 			Setting.BoneFilterActionOption = EBoneFilterActionOption::Remove;
