Bring missing changes from Release-Engine-Staging to Release-Engine-Test

diff --git a/Engine/Plugins/Bridge/Bridge.uplugin b/Engine/Plugins/Bridge/Bridge.uplugin
--- a/Engine/Plugins/Bridge/Bridge.uplugin
+++ b/Engine/Plugins/Bridge/Bridge.uplugin
@@ -10,6 +10,7 @@
 	"DocsURL": "https://help.quixel.com/hc/en-us/sections/360005846137-Quixel-Bridge-for-Unreal-Engine-5",
 	"MarketplaceURL": "",
 	"SupportURL": "[REDACTED_EMAIL]",
+	"EngineVersion": "5.0.0",
 	"CanContainContent": true,
 	"EnabledByDefault": true,
 	"SupportedTargetPlatforms": [
diff --git a/Engine/Plugins/Experimental/SkeletalReduction/Source/Private/SkeletalSimplifierVertex.h b/Engine/Plugins/Experimental/SkeletalReduction/Source/Private/SkeletalSimplifierVertex.h
--- a/Engine/Plugins/Experimental/SkeletalReduction/Source/Private/SkeletalSimplifierVertex.h
+++ b/Engine/Plugins/Experimental/SkeletalReduction/Source/Private/SkeletalSimplifierVertex.h
@@ -37,11 +37,11 @@ namespace SkeletalSimplifier
 			// NB: Required that these all have float storage.
 			// - Base Attributes: size = 13 + 2 * NumTexCoord
 
-			FVector3f			Normal;      // 0, 1, 2
-			FVector3f			Tangent;     // 3, 4, 5
-			FVector3f         BiTangent;   // 6, 7, 8
+			FVector3f		Normal;      // 0, 1, 2
+			FVector3f		Tangent;     // 3, 4, 5
+			FVector3f       BiTangent;   // 6, 7, 8
 			FLinearColor	Color;       // 9, 10, 11, 12
-			FVector2D		TexCoords[NumTexCoords];  // 13, .. 13 + NumTexCoords * 2 - 1
+			FVector2f		TexCoords[NumTexCoords];  // 13, .. 13 + NumTexCoords * 2 - 1
 
 
 			// used to manage identity of split/non-split vertex attributes.
@@ -157,7 +157,7 @@ namespace SkeletalSimplifier
 			{
 				for (int32 i = 0; i < NumTexCoords; ++i)
 				{
-					TexCoords[i] = FVector2D(ForceInitToZero);
+					TexCoords[i] = FVector2f(ForceInitToZero);
 				}
 			}
 
diff --git a/Engine/Source/Developer/Apple/MetalShaderFormat/Private/MetalShaderFormat.cpp b/Engine/Source/Developer/Apple/MetalShaderFormat/Private/MetalShaderFormat.cpp
--- a/Engine/Source/Developer/Apple/MetalShaderFormat/Private/MetalShaderFormat.cpp
+++ b/Engine/Source/Developer/Apple/MetalShaderFormat/Private/MetalShaderFormat.cpp
@@ -389,7 +389,7 @@ static FMetalCompilerToolchain::EMetalToolchainStatus ParseCompilerVersionAndTar
 		char AppleToolName[PATH_MAX] = { '\0' };
 		NumResults = sscanf(TCHAR_TO_ANSI(*Version), "Apple %s version %d.%d.%d", AppleToolName, &Major, &Minor, &Patch);
 #else
-		TCHAR AppleToolName[WINDOWS_MAX_PATH] = TEXT("");
+		TCHAR AppleToolName[WINDOWS_MAX_PATH] = { '\0' };
 		NumResults = swscanf_s(*Version, TEXT("Apple %ls version %d.%d.%d"), AppleToolName, WINDOWS_MAX_PATH, &Major, &Minor, &Patch);
 #endif
 		PackedVersionNumber.Major = Major;
