Fix the alternate skinning on generated LOD to adjust the section max bone influence count in case the alternate skinning use more influence then the generated mesh skinning.
#rnx

diff --git a/Engine/Plugins/Experimental/SkeletalReduction/Source/Private/SkeletalMeshReductionPlugin.cpp b/Engine/Plugins/Experimental/SkeletalReduction/Source/Private/SkeletalMeshReductionPlugin.cpp
--- a/Engine/Plugins/Experimental/SkeletalReduction/Source/Private/SkeletalMeshReductionPlugin.cpp
+++ b/Engine/Plugins/Experimental/SkeletalReduction/Source/Private/SkeletalMeshReductionPlugin.cpp
@@ -1592,21 +1592,22 @@ void  FQuadricSkeletalMeshReduction::AddSourceModelInfluences( const FSkeletalMe
 			// map the verts in the LOD model to the input order.
 			const TArray<int32>& ImportVertexMap = LODModel.MeshToImportVertexMap;
 			check(ImportVertexMap.Num() == (int32)LODModel.NumVertices);
-
-			for (int32 i = 0; i < (int32)LODModel.NumVertices; ++i)
+			int32 MaxNumInfluences = 0;
+			for (int32 VertexIndex = 0; VertexIndex < (int32)LODModel.NumVertices; ++VertexIndex)
 			{
 				// Map to section and to imported vertex
 
-				const int32 SectionId = VertIdxToSectionMap[i];
-				const int32 SrcVertId = ImportVertexMap[i];
+				const int32 SectionId = VertIdxToSectionMap[VertexIndex];
+				const int32 SrcVertId = ImportVertexMap[VertexIndex];
 
+				FSkelMeshSection& Section = LODModel.Sections[SectionId];
 				// The BoneMap for this section, needed to encode bones.
 
-				const auto& BoneMap = LODModel.Sections[SectionId].BoneMap;
+				const auto& BoneMap = Section.BoneMap;
 
 				// the dst for the bones and weights.
 
-				FRawSkinWeight& WeightAndBones = SkinWeights[i];
+				FRawSkinWeight& WeightAndBones = SkinWeights[VertexIndex];
 				ZeroFRawSkinWeight(WeightAndBones);
 
 				// The bones and Weights for this vertex.
@@ -1616,7 +1617,7 @@ void  FQuadricSkeletalMeshReduction::AddSourceModelInfluences( const FSkeletalMe
 					// Add each bone / weight. 
 					// keep track of the total weight.  should sum to 255 and the first weight is the largest
 					int32 TotalQuantizedWeight = 0;
-					int32 b = 0;
+					int32 InfluenceIndex = 0;
 					for (const auto& Pair : BoneWeight.GetData())
 					{
 						int32 BoneId = Pair.Key;
@@ -1625,7 +1626,7 @@ void  FQuadricSkeletalMeshReduction::AddSourceModelInfluences( const FSkeletalMe
 						// Transform weight to quantized weight
 						uint8 QuantizedWeight = FMath::Clamp((uint8)(Weight*((double)0xFF)), (uint8)0x00, (uint8)0xFF);
 
-						WeightAndBones.InfluenceWeights[b] = QuantizedWeight;
+						WeightAndBones.InfluenceWeights[InfluenceIndex] = QuantizedWeight;
 
 						TotalQuantizedWeight += QuantizedWeight;
 
@@ -1634,16 +1635,24 @@ void  FQuadricSkeletalMeshReduction::AddSourceModelInfluences( const FSkeletalMe
 						int32 LocalBoneId = BoneMap.Find(BoneId);
 						if (LocalBoneId != INDEX_NONE)
 						{
-							WeightAndBones.InfluenceBones[b] = LocalBoneId;
+							WeightAndBones.InfluenceBones[InfluenceIndex] = LocalBoneId;
 						}
 						else
 						{
 							// Map to root of section
-							WeightAndBones.InfluenceBones[b] = 0;
+							WeightAndBones.InfluenceBones[InfluenceIndex] = 0;
 
 							check(0); // should never hit this
 						}
-						b++;
+						InfluenceIndex++;
+					}
+					if (InfluenceIndex > MaxNumInfluences)
+					{
+						MaxNumInfluences = InfluenceIndex;
+						if (MaxNumInfluences > Section.GetMaxBoneInfluences())
+						{
+							Section.MaxBoneInfluences = MaxNumInfluences;
+						}
 					}
 					//Use the same code has the build where we modify the index 0 to have a sum of 255 for all influence per skin vertex
 					int32 ExcessQuantizedWeight = 255 - TotalQuantizedWeight;
@@ -1651,7 +1660,7 @@ void  FQuadricSkeletalMeshReduction::AddSourceModelInfluences( const FSkeletalMe
 					WeightAndBones.InfluenceWeights[0] += ExcessQuantizedWeight;
 				}
 			}
-
+			
 		}
 
 	}
