SkeletalMeshReductionSkinnedMesh - clean up the logic in the Empty() method.

Based on github pull request from ss (ssdemajia)

#9205

diff --git a/Engine/Plugins/Experimental/SkeletalReduction/Source/Private/SkeletalMeshReductionSkinnedMesh.h b/Engine/Plugins/Experimental/SkeletalReduction/Source/Private/SkeletalMeshReductionSkinnedMesh.h
--- a/Engine/Plugins/Experimental/SkeletalReduction/Source/Private/SkeletalMeshReductionSkinnedMesh.h
+++ b/Engine/Plugins/Experimental/SkeletalReduction/Source/Private/SkeletalMeshReductionSkinnedMesh.h
@@ -52,8 +52,8 @@ namespace SkeletalSimplifier
 		*/
 		void Resize(int32 NumTriangles, int32 NumVertices)
 		{
-			if (IndexBuffer == NULL)  delete[] IndexBuffer;
-			if (VertexBuffer == NULL) delete[] VertexBuffer;
+			delete[] IndexBuffer;
+			delete[] VertexBuffer;
 
 			NumTris  = NumTriangles;
 			NumVerts = NumVertices;
@@ -67,8 +67,11 @@ namespace SkeletalSimplifier
 		*/
 		void Empty()
 		{
-			if (IndexBuffer == NULL)  delete[] IndexBuffer;
-			if (VertexBuffer == NULL) delete[] VertexBuffer;
+			delete[] IndexBuffer;
+			delete[] VertexBuffer;
+
+			IndexBuffer =  NULL;
+			VertexBuffer = NULL;
 
 			NumTris = 0;
 			NumVerts = 0;
