SkeletalReduction - inclusive terminology update.  This is a trivial rename.

diff --git a/Engine/Plugins/Experimental/SkeletalReduction/Source/Private/SkeletalMeshReductionPlugin.cpp b/Engine/Plugins/Experimental/SkeletalReduction/Source/Private/SkeletalMeshReductionPlugin.cpp
--- a/Engine/Plugins/Experimental/SkeletalReduction/Source/Private/SkeletalMeshReductionPlugin.cpp
+++ b/Engine/Plugins/Experimental/SkeletalReduction/Source/Private/SkeletalMeshReductionPlugin.cpp
@@ -764,7 +764,7 @@ void FQuadricSkeletalMeshReduction::ConvertToFSkinnedSkeletalMesh( const FSkelet
 			}
 
 			BasicAttrs.Color = SkinnedVertex.Color;
-			OutVertexBuffer[v].MasterVertIndex = v; // index of the closest vert w.r.t  SrcLODModel.GetVertices(SoftSkinVertices);
+			OutVertexBuffer[v].ClosestSrcVertIndex = v; // index of the closest vert w.r.t  SrcLODModel.GetVertices(SoftSkinVertices);
 			OutVertexBuffer[v].MaterialIndex   = 0; // default to be over-written
 			OutVertexBuffer[v].Position = SkinnedVertex.Position;
 
@@ -1478,14 +1478,14 @@ void  FQuadricSkeletalMeshReduction::AddSourceModelInfluences( const FSkeletalMe
 			for (int32 VIdx = 0; VIdx < SkinnedMesh.NumVertices(); ++VIdx)
 			{
 				// The VerId in the Source Mesh that was closest to the simplified vertex
-				const int32 MasterVertId = SkinnedMesh.VertexBuffer[VIdx].MasterVertIndex;
+				const int32 ClosestSrcVertId = SkinnedMesh.VertexBuffer[VIdx].ClosestSrcVertIndex;
 
-				if (MasterVertId != INDEX_NONE)  // NB: MasterVerId should never be INDEX_NONE
+				if (ClosestSrcVertId != INDEX_NONE)  // NB: ClosestSrcVertId should never be INDEX_NONE
 				{
-					const FRawSkinWeight& SrcRawSkinWeight = SrcBonesAndWeights[MasterVertId];
+					const FRawSkinWeight& SrcRawSkinWeight = SrcBonesAndWeights[ClosestSrcVertId];
 
 					// Get the BoneMap that was used to encode these weights. 
-					const auto& BoneMap = SrcLODModel.Sections[SrcVertIdxToSectionMap[MasterVertId]].BoneMap;
+					const auto& BoneMap = SrcLODModel.Sections[SrcVertIdxToSectionMap[ClosestSrcVertId]].BoneMap;
 
 					// Add the non-zero weights 
 					for (int32 b = 0; b < MAX_TOTAL_INFLUENCES; ++b)
diff --git a/Engine/Plugins/Experimental/SkeletalReduction/Source/Private/SkeletalSimplifier.cpp b/Engine/Plugins/Experimental/SkeletalReduction/Source/Private/SkeletalSimplifier.cpp
--- a/Engine/Plugins/Experimental/SkeletalReduction/Source/Private/SkeletalSimplifier.cpp
+++ b/Engine/Plugins/Experimental/SkeletalReduction/Source/Private/SkeletalSimplifier.cpp
@@ -306,14 +306,14 @@ void SkeletalSimplifier::FMeshSimplifier::ComputeEdgeCollapseVertsAndFixBones(Si
 	const auto& ClosestSimpliferVert = (DstSqr1 < DstSqr0) ? Vert1->vert : Vert0->vert;
 
 	const auto& SrcBones         = ClosestSimpliferVert.GetSparseBones();
-	const int32 MasterVertIndex  = ClosestSimpliferVert.MasterVertIndex;
+	const int32 ClosestSrcVertIndex  = ClosestSimpliferVert.ClosestSrcVertIndex;
 	const int32 NumNewVerts      = EdgeAndNewVertArray.Num();
 
 	for (int32 i = 0; i < NumNewVerts; ++i)
 	{
 		MeshVertType& simpVert   = EdgeAndNewVertArray[i].Get<2>();
 		simpVert.SparseBones     = SrcBones;
-		simpVert.MasterVertIndex = MasterVertIndex;
+		simpVert.ClosestSrcVertIndex = ClosestSrcVertIndex;
 	}
 
 }
diff --git a/Engine/Plugins/Experimental/SkeletalReduction/Source/Private/SkeletalSimplifierVertex.h b/Engine/Plugins/Experimental/SkeletalReduction/Source/Private/SkeletalSimplifierVertex.h
--- a/Engine/Plugins/Experimental/SkeletalReduction/Source/Private/SkeletalSimplifierVertex.h
+++ b/Engine/Plugins/Experimental/SkeletalReduction/Source/Private/SkeletalSimplifierVertex.h
@@ -410,7 +410,7 @@ namespace SkeletalSimplifier
 			//   i.e. this number should be inherited from the closest vert during edge collapse.
 			//   the value -1 will be used as a null-flag
 			
-			int32 MasterVertIndex;
+			int32 ClosestSrcVertIndex;
 			
 			// The Vertex Point
 			uint32			MaterialIndex;
@@ -440,7 +440,7 @@ namespace SkeletalSimplifier
 
 
 			TSkeletalSimpVert() :
-				MasterVertIndex(-1),
+				ClosestSrcVertIndex(-1),
 				MaterialIndex(0),
 				Position(ForceInitToZero),
 				SpecializedWeight(0.f),
@@ -452,7 +452,7 @@ namespace SkeletalSimplifier
 
 			// copy constructor
 			TSkeletalSimpVert(const TSkeletalSimpVert& Other) :
-				MasterVertIndex(Other.MasterVertIndex),
+				ClosestSrcVertIndex(Other.ClosestSrcVertIndex),
 				MaterialIndex(Other.MaterialIndex),
 				Position(Other.Position),
 				SpecializedWeight(Other.SpecializedWeight),
@@ -502,7 +502,7 @@ namespace SkeletalSimplifier
 
 			TSkeletalSimpVert& operator=(const TSkeletalSimpVert& Other)
 			{
-				MasterVertIndex       = Other.MasterVertIndex;
+				ClosestSrcVertIndex  = Other.ClosestSrcVertIndex;
 				MaterialIndex        = Other.MaterialIndex;
 				Position             = Other.Position;
 				SpecializedWeight    = Other.SpecializedWeight;
@@ -521,7 +521,7 @@ namespace SkeletalSimplifier
 			{
 				bool bIsApprxEquals = 
 					(MaterialIndex == Other.MaterialIndex)
-                 && (MasterVertIndex == Other.MasterVertIndex)
+                 && (ClosestSrcVertIndex == Other.ClosestSrcVertIndex)
 			     && PointsEqual(Position, Other.Position);
 
 				bIsApprxEquals = bIsApprxEquals
@@ -543,7 +543,7 @@ namespace SkeletalSimplifier
 			bool operator==(const VertType& Other) const
 			{
 				bool bIsEqual = (MaterialIndex == Other.MaterialIndex) && 
-					            (MasterVertIndex == Other.MasterVertIndex) &&
+					            (ClosestSrcVertIndex == Other.ClosestSrcVertIndex) &&
 					            (Position == Other.Position);
 				bIsEqual = bIsEqual && (SpecializedWeight == Other.SpecializedWeight);
 				bIsEqual = bIsEqual && (GetBasicAttrAccessor() == Other.GetBasicAttrAccessor());
@@ -554,7 +554,7 @@ namespace SkeletalSimplifier
 			}
 
 			// Standard operator overloading.
-			// Note: these don't affect the MasterVertIndex or MaterialIndex
+			// Note: these don't affect the ClosestSrcVertIndex or MaterialIndex
 
 			VertType operator+(const VertType& Other) const
 			{
