Fix a non normalized skin weight issue

#rnx

diff --git a/Engine/Plugins/Experimental/SkeletalReduction/Source/Private/SkeletalMeshReductionPlugin.cpp b/Engine/Plugins/Experimental/SkeletalReduction/Source/Private/SkeletalMeshReductionPlugin.cpp
--- a/Engine/Plugins/Experimental/SkeletalReduction/Source/Private/SkeletalMeshReductionPlugin.cpp
+++ b/Engine/Plugins/Experimental/SkeletalReduction/Source/Private/SkeletalMeshReductionPlugin.cpp
@@ -219,6 +219,7 @@ private:
 		                                const TArray<FMatrix>& BoneMatrices,
 		                                const int32 LODIndex,
 		                                SkeletalSimplifier::FSkinnedSkeletalMesh& OutSkinnedMesh,
+										const FString& SkeletalMeshName,
 										TArray<int32>* VertToSectionMap = nullptr) const;
 
 	/**
@@ -490,6 +491,7 @@ void FQuadricSkeletalMeshReduction::ConvertToFSkinnedSkeletalMesh( const FSkelet
 	                                                               const TArray<FMatrix>& BoneMatrices,
 	                                                               const int32 LODIndex,
 	                                                               SkeletalSimplifier::FSkinnedSkeletalMesh& OutSkinnedMesh,
+																   const FString& SkeletalMeshName,
 																   TArray<int32>* VertToSectionMap) const
 {
 
@@ -674,13 +676,13 @@ void FQuadricSkeletalMeshReduction::ConvertToFSkinnedSkeletalMesh( const FSkelet
 		// Report any error with invalid bone weights
 		if (!bHasValidBoneWeights && !SkipSection(SectionIndex))
 		{
-			UE_LOG(LogSkeletalMeshReduction, Warning, TEXT("Building LOD %d - Encountered questionable vertex weights in source."), LODIndex);
+			UE_LOG(LogSkeletalMeshReduction, Warning, TEXT("Building LOD %d - Encountered non normalized vertex weights in source. %s"), LODIndex, *SkeletalMeshName);
 		}
 	}
 
 	if (NumBadNTBSpace > 0)
 	{
-		UE_LOG(LogSkeletalMeshReduction, Warning, TEXT("There were NaNs in the Tangent Space of %d source model vertices."), NumBadNTBSpace);
+		UE_LOG(LogSkeletalMeshReduction, Warning, TEXT("There were NaNs in the Tangent Space of %d source model vertices. %s"), NumBadNTBSpace, *SkeletalMeshName);
 	}
 
 	// -- Make the index buffer; skipping the "SkipSections"
@@ -1786,7 +1788,7 @@ bool FQuadricSkeletalMeshReduction::ReduceSkeletalLODModel( const FSkeletalMeshL
 	// Generate a single skinned mesh form the SrcModel.  This mesh has per-vertex tangent space.
 	TArray<int32> VertToOriginalSectionId; 
 	SkeletalSimplifier::FSkinnedSkeletalMesh SkinnedSkeletalMesh;
-	ConvertToFSkinnedSkeletalMesh(SrcModel, BoneMatrices, LODIndex, SkinnedSkeletalMesh, &VertToOriginalSectionId);
+	ConvertToFSkinnedSkeletalMesh(SrcModel, BoneMatrices, LODIndex, SkinnedSkeletalMesh, SkeletalMeshName, &VertToOriginalSectionId);
 
 	int32 IterationNum = 0;
 	//We keep the original MaxNumVerts because if we iterate we want to still compare with the original request.
@@ -1852,7 +1854,7 @@ bool FQuadricSkeletalMeshReduction::ReduceSkeletalLODModel( const FSkeletalMeshL
 				}
 
 				UE_LOG(LogSkeletalMeshReduction, Log, TEXT("Chunking to limit unique bones per section generated additional vertices - continuing simplification of LOD %d "), LODIndex);
-				ConvertToFSkinnedSkeletalMesh(SrcModel, BoneMatrices, LODIndex, SkinnedSkeletalMesh);
+				ConvertToFSkinnedSkeletalMesh(SrcModel, BoneMatrices, LODIndex, SkinnedSkeletalMesh, SkeletalMeshName);
 			}
 			else
 			{
@@ -1958,7 +1960,7 @@ void FQuadricSkeletalMeshReduction::ReduceSkeletalMesh(USkeletalMesh& SkeletalMe
 		if (Settings.BaseLOD == LODIndex && !SkeletalMesh.IsLODImportedDataBuildAvailable(LODIndex))
 		{
 			//Cannot reduce ourself if we are not imported
-			UE_LOG(LogSkeletalMeshReduction, Warning, TEXT("Building LOD %d - Cannot generate LOD with itself if the LOD do not have imported Data. Using Base LOD 0 instead"), LODIndex);
+			UE_ASSET_LOG(LogSkeletalMeshReduction, Warning, &SkeletalMesh, TEXT("Building LOD %d - Cannot generate LOD with itself if the LOD do not have imported Data. Using Base LOD 0 instead"), LODIndex);
 		}
 		else if (Settings.BaseLOD <= LODIndex && SkeletalMeshResource.LODModels.IsValidIndex(Settings.BaseLOD))
 		{
@@ -1968,7 +1970,7 @@ void FQuadricSkeletalMeshReduction::ReduceSkeletalMesh(USkeletalMesh& SkeletalMe
 		else
 		{
 			// warn users
-			UE_LOG(LogSkeletalMeshReduction, Warning, TEXT("Building LOD %d - Invalid Base LOD entered. Using Base LOD 0 instead"), LODIndex);
+			UE_ASSET_LOG(LogSkeletalMeshReduction, Warning, &SkeletalMesh, TEXT("Building LOD %d - Invalid Base LOD entered. Using Base LOD 0 instead"), LODIndex);
 		}
 	}
 
