Fixed mismatched copy constructors/assignment operators in SkeletalSimplifier::LinearAlgebra::ArrayBase and derived classes.

diff --git a/Engine/Plugins/Experimental/SkeletalReduction/Source/Private/SkeletalSimplifierLinearAlgebra.h b/Engine/Plugins/Experimental/SkeletalReduction/Source/Private/SkeletalSimplifierLinearAlgebra.h
--- a/Engine/Plugins/Experimental/SkeletalReduction/Source/Private/SkeletalSimplifierLinearAlgebra.h
+++ b/Engine/Plugins/Experimental/SkeletalReduction/Source/Private/SkeletalSimplifierLinearAlgebra.h
@@ -37,12 +37,6 @@ namespace SkeletalSimplifier
 				memset(Data, 0, sizeof(Data));
 			}
 
-			ArrayBase(const ArrayBase& other)
-			{
-				//for (int i = 0; i < SIZE; ++i) Data[i] = other.Data[i];
-				memcpy(Data, other.Data, sizeof(Data));
-			}
-
 			// Data access
 			ScalarType& operator() (const int32 idx)
 			{
@@ -271,10 +265,6 @@ namespace SkeletalSimplifier
 				                                Data[5] = a33;
 			}
 
-			SymmetricMatrix(const SymmetricMatrix& other) :
-				MyBase(other)
-			{}
-
 			SymmetricMatrix(const MyBase& base) :
 				MyBase(base)
 			{}
@@ -492,10 +482,6 @@ namespace SkeletalSimplifier
 				Data[6] = Row2[0];  Data[7] = Row2[1];  Data[8] = Row2[2];
 			}
 
-			DMatrix(const DMatrix& Other) :
-				MyBase(Other)
-			{}
-
 			/**
 			* Accesses elements in the matrix using standard M(i,j) notation
 
