Skelmesh: Add missing copy operator.

#rnx

diff --git a/Engine/Plugins/Experimental/SkeletalReduction/Source/Private/SkeletalSimplifierLinearAlgebra.h b/Engine/Plugins/Experimental/SkeletalReduction/Source/Private/SkeletalSimplifierLinearAlgebra.h
--- a/Engine/Plugins/Experimental/SkeletalReduction/Source/Private/SkeletalSimplifierLinearAlgebra.h
+++ b/Engine/Plugins/Experimental/SkeletalReduction/Source/Private/SkeletalSimplifierLinearAlgebra.h
@@ -985,7 +985,8 @@ namespace SkeletalSimplifier
 				}
 			}
 
-			TDenseVecD(const TDenseVecD& Other) : MyBase(Other) {};
+			TDenseVecD(const TDenseVecD& Other) = default;
+			TDenseVecD& operator=(const TDenseVecD& Other) = default;
 
 			void SetElement(const int32 j, const double Value) { MyBase::operator[](j) = Value; }
 			ScalarType GetElement(const int32 j) const { return MyBase::operator[](j); }
@@ -1059,15 +1060,6 @@ namespace SkeletalSimplifier
 				return GetColumn(j)[i];
 			}
 
-			/**
-			* Assignment.
-			*/
-			SparseBMatrix& operator=(const SparseBMatrix& Other)
-			{
-				SparseData = Other.SparseData;
-				return *this;
-			}
-
 			/**
 			* Add another SparseBMatrix to this one
 			* Element by element addition.
diff --git a/Engine/Plugins/Experimental/SkeletalReduction/Source/Private/SkeletalSimplifierQuadrics.h b/Engine/Plugins/Experimental/SkeletalReduction/Source/Private/SkeletalSimplifierQuadrics.h
--- a/Engine/Plugins/Experimental/SkeletalReduction/Source/Private/SkeletalSimplifierQuadrics.h
+++ b/Engine/Plugins/Experimental/SkeletalReduction/Source/Private/SkeletalSimplifierQuadrics.h
@@ -25,12 +25,9 @@ namespace SkeletalSimplifier
 
 			FEdgeQuadric(const Vec3d& Vert0Pos, const Vec3d& Vert1Pos, const Vec3d& FaceNormal, const double EdgeWeight);
 
-			FEdgeQuadric(const FEdgeQuadric& Other) :
-				CMatrix(Other.CMatrix),
-				D0Vector(Other.D0Vector),
-				CScalar(Other.CScalar)
-			{}
-
+			FEdgeQuadric(const FEdgeQuadric&) = default;
+			FEdgeQuadric& operator=(const FEdgeQuadric&) = default;
+			
 			double Evaluate(const Vec3d& Pos) const;
 			
 
@@ -200,14 +197,9 @@ namespace SkeletalSimplifier
 				         const SparseWeightContainerType& AdditionalWeights);
 
 
-			TFaceQuadric() :
-				MyBase()
-			{}
-
-
-			TFaceQuadric(const TFaceQuadric& Other) :
-				MyBase(Other)
-			{}
+			TFaceQuadric() = default;
+			TFaceQuadric(const TFaceQuadric& Other) = default;
+			TFaceQuadric& operator=(const TFaceQuadric& Other) = default;
 
 			/**
 			*  Returns the quadric evaluation at this point  (vTr . A . v  + 2 bTr. v + c)  where A = quadric matrix, b = DVector, c = CScalar
