Skeletal Mesh Reduction.  Improve attribute interpolation during simplification.

Previously,  when an edge is collasped and two triangles are replaced by a single vertex, the quality of the attributes on that new vertex (e.g. UVs ) was limited for the special case when one of the triangles is on a plane that intersects the origin.  This change fixes that special case.

diff --git a/Engine/Plugins/Experimental/SkeletalReduction/Source/Private/SkeletalSimplifierLinearAlgebra.h b/Engine/Plugins/Experimental/SkeletalReduction/Source/Private/SkeletalSimplifierLinearAlgebra.h
--- a/Engine/Plugins/Experimental/SkeletalReduction/Source/Private/SkeletalSimplifierLinearAlgebra.h
+++ b/Engine/Plugins/Experimental/SkeletalReduction/Source/Private/SkeletalSimplifierLinearAlgebra.h
@@ -1345,64 +1345,63 @@ namespace SkeletalSimplifier
 		* This class generates the interpolation coefficients vector (Vec3d) g  and distance (double) d
 		* defined over the face of a triangle.
 		*
-		* The Position Matrix is defined in terms of the three corners of triangle
-		*  {pa, pb, pc} with corresponding normal 'FaceNormal'
-		*
-		* Position Matrix =
-		*   ( pa_0, pa_1, pa_2 )
-		*   ( pb_0, pb_1, pb_2 )
-		*   ( pc_0, pc_1, pc_2 )
-		*
+		* *
 		* The actual system solved is:
 		*  <pa | g> + d = s0
 		*  <pb | g> + d = s1
 		*  <pc | g> + d = s2
 		*  <FaceNormal | g> = 0
 		*
+		* But we re-write this as
+		* <(pc - pa) | g >  = s2 - s0
+		* <(pb - pa) | g >  = s1 - s0
+		* < FaceNormal | g > = 0
+		* d = s0 - <pa | g>
+		*
+		* The first three equations can be solved for "g", in terms of s, and the result can be used to
+		* compute d.
 		*
-		* In matrix form   ( PositionMatrix   Vec3(1) )   ( g )  = ( s )
-		*                  ( FaceNormal^T ,     0     )   ( d )    ( 0 )
+		* The BasisMatrix is defined in terms of the three corners of triangle
+		*  {pa, pb, pc} and the normal corresponding normal 'FaceNormal'
 		*
-		* where the vector (Vec3d) 's' represents the per-vertex data that forms boundary conditions
-		* for the interpolation.
+		* BasisMatrix =
+		*   ( pc_0 - pa_0 , pc_1 - pa_1 , pc_2 - pa_2 )
+		*   ( pb_0 - pa_0 , pb_1 - pa_1 , pb_2 - pa_2)
+		*   ( FaceNormal_0, FaceNormal_1, FaceNormal_2 )
+		* The rows corresponds to vectors along two sides of the triangle and the face normal
+		*
+		* Solving for "g" is simply solving
+		*  BasisMatrix * g = vec(s2-s0, s1-s0, 0)
+		*
+		* or
+		*  g  = Invserse(BasisMatrix) * vec(s2-s0, s1-s0, 0)
+		* the Inverse(BasisMatrix) is cached as the CoBasisMatrix
 		*
-		*  The actual solution is given by:
-		*  -- Distance
-		*  double d = < FaceNormal | InvsPositionMatrix * s> / <FaceNormal | InversePositionMatrix * Vec3d(1) >;
-		*  -- Gradient
-		*  Vec3d  g = InversePositionMatrix * s - d * InversePositionMatrix * Vec3d(1);
 		*
 		* The computation is broken up do allow for reuse with multiple sets of per-vertex data.
 		*/
 		class InverseGradientProjection
 		{
 		public:
-			InverseGradientProjection(const DMatrix& PositionMatrix, const Vec3d& FaceNormal)
+			InverseGradientProjection(const Vec3d& Pos0, const Vec3d& Pos1, const Vec3d& Pos2, const Vec3d& FaceNormal)
 			{
-				// Threshold for computing the matrix inverse.
-				const double DetThreshold = 1.e-8;
 
-				// Compute the inverse of the position matrix
-				PosInv = PositionMatrix.Inverse(bIsValid, DetThreshold);
-				if (bIsValid)
-				{
-					// InversePositionMatrix * Vec3(1)
-					MInv1 = PosInv.RowSum();
+				// Threshold for computing the matrix inverse.
+				constexpr  double DetThreshold = 1.e-8;
 
-					// <FaceNormal | InversePositionMatrix 
-					Dhat = FaceNormal * PosInv;
+				Origin = Pos0;
 
-					// <FaceNormal | InvesePositionMatrix Vec3d(1)> 
-					double ReScale = Dhat[0] + Dhat[1] + Dhat[2];
+				const Vec3d Side20 = Pos2 - Pos0;
+				const Vec3d Side10 = Pos1 - Pos0;
 
-					bIsValid = bIsValid && (FMath::Abs(ReScale) > 1.e-8);
+				DMatrix BasisMatrix(Side20,
+									Side10,
+									FaceNormal);
 
-					// divide by <FaceNormal | InvesePositionMatrix Vec3d(1)> 
-					Dhat *= ( 1. / ReScale );
-					//now:  Dhat =  <FaceNormal | InvPos  / <FaceNormal | InvPos. Vec3d(1) >
-				}
+				CoBasisMatrix = BasisMatrix.Inverse(bIsValid, DetThreshold);
 			}
 
+
 			bool IsValid() const
 			{
 				return bIsValid;
@@ -1413,14 +1412,10 @@ namespace SkeletalSimplifier
 			// @return  Distance
 			double ComputeGradient(const Vec3d& PerVertexData, Vec3d& OutGradient) const
 			{
-				// PosInv . s
-				Vec3d MInvS = PosInv * PerVertexData;
+				Vec3d DataVec(PerVertexData[2] - PerVertexData[0], PerVertexData[1] - PerVertexData[0], 0.);
 
-				// Dist = <dhat | s>
-				double Distance = Dhat.DotProduct(PerVertexData);
-				
-				// Grad =  PosInv . s - Dist * PosInv . (1, 1, 1} 
-				OutGradient = MInvS - Distance * MInv1;
+				OutGradient = CoBasisMatrix * DataVec;
+				double Distance = PerVertexData[0] - Origin.DotProduct(OutGradient);
 
 				return Distance;
 			}
@@ -1428,9 +1423,8 @@ namespace SkeletalSimplifier
 		private:
 
 			bool    bIsValid;  // Indicates if the inversions incurred a divide by very-small-number.
-			DMatrix PosInv;    // Inverse(PositionMatrix)
-			Vec3d   Dhat;      // n * Inverse(PositinoMatrix) / < n | Inverse(PositionMatrix) Vec3(1) >
-			Vec3d   MInv1;     // Inverse(PositionMatrix) * Vec3(1)
+			DMatrix CoBasisMatrix;   // Inverse( {P2-P0}, {P1 - P0}, Normal}
+			Vec3d   Origin;    // corresponds to P0
 		};
 
 		template <typename SparseVecDOrDenseVecD>
diff --git a/Engine/Plugins/Experimental/SkeletalReduction/Source/Private/SkeletalSimplifierQuadrics.h b/Engine/Plugins/Experimental/SkeletalReduction/Source/Private/SkeletalSimplifierQuadrics.h
--- a/Engine/Plugins/Experimental/SkeletalReduction/Source/Private/SkeletalSimplifierQuadrics.h
+++ b/Engine/Plugins/Experimental/SkeletalReduction/Source/Private/SkeletalSimplifierQuadrics.h
@@ -605,10 +605,10 @@ namespace SkeletalSimplifier
 
 															// create the tool needed to compute the gradient coefficients of the attributes.
 
-			InverseGradientProjection  GradientTool(DMatrix(Vert0Pos,
-			                                            	Vert1Pos,
-				                                            Vert2Pos),
-				                                            FaceNormal);
+			InverseGradientProjection  GradientTool(Vert0Pos,
+													Vert1Pos,
+													Vert2Pos,
+													FaceNormal);
 
 			// Accumulate the terms related to the gradient of the Basic Attributes
 			{
