SkeletalMesh Simplification.  Adding the ability to select the original version of an attribute gradient that results in better triangle density for cloth but potentially worse UVs (on by default), or an improved attribute gradient.

This primarily affects 2d meshes that are on a plane that intersects the origin.

diff --git a/Engine/Plugins/Experimental/SkeletalReduction/Source/Private/SkeletalMeshReductionPlugin.cpp b/Engine/Plugins/Experimental/SkeletalReduction/Source/Private/SkeletalMeshReductionPlugin.cpp
--- a/Engine/Plugins/Experimental/SkeletalReduction/Source/Private/SkeletalMeshReductionPlugin.cpp
+++ b/Engine/Plugins/Experimental/SkeletalReduction/Source/Private/SkeletalMeshReductionPlugin.cpp
@@ -1044,6 +1044,7 @@ float FQuadricSkeletalMeshReduction::SimplifyMesh( const FSkeletalMeshOptimizati
                                                   ) const
 {
 
+	const bool bUseLegacyAttrGrad = Settings.bImproveTrianglesForCloth;
 	// Convert settings to weights and a termination criteria
 
 	// Determine the stop criteria used
@@ -1192,7 +1193,8 @@ float FQuadricSkeletalMeshReduction::SimplifyMesh( const FSkeletalMeshOptimizati
 	SkeletalSimplifier::FMeshSimplifier  Simplifier(Mesh.VertexBuffer, (uint32)Mesh.NumVertices(),
 		                                            Mesh.IndexBuffer, (uint32)Mesh.NumIndices(), 
 		                                            CoAlignmentLimit, VolumeImportance, bPreserveVolume,  
-													bEnforceBoneBoundaries, bMergeCoincidentVertBones);
+													bEnforceBoneBoundaries, bMergeCoincidentVertBones,
+													bUseLegacyAttrGrad);
 
 	// The simplifier made a deep copy of the mesh.  
 
diff --git a/Engine/Plugins/Experimental/SkeletalReduction/Source/Private/SkeletalSimplifier.cpp b/Engine/Plugins/Experimental/SkeletalReduction/Source/Private/SkeletalSimplifier.cpp
--- a/Engine/Plugins/Experimental/SkeletalReduction/Source/Private/SkeletalSimplifier.cpp
+++ b/Engine/Plugins/Experimental/SkeletalReduction/Source/Private/SkeletalSimplifier.cpp
@@ -15,13 +15,14 @@ const int32 SkeletalSimplifier::LinearAlgebra::SymmetricMatrix::Mapping[9] = { 0
 SkeletalSimplifier::FMeshSimplifier::FMeshSimplifier(const MeshVertType* InSrcVerts, const uint32 InNumSrcVerts,
 	const uint32* InSrcIndexes, const uint32 InNumSrcIndexes,
 	const float CoAlignmentLimit, const float VolumeImportanceValue, const bool VolumeConservation, const bool bEnforceBoundaries, 
-	const bool bMergeCoincidentVertBones)
+	const bool bMergeCoincidentVertBones, const bool bUseLegacyAttrGradient)
 	:
 	coAlignmentLimit(CoAlignmentLimit),
 	VolumeImportance(VolumeImportanceValue),
 	bPreserveVolume(VolumeConservation),
 	bCheckBoneBoundaries(bEnforceBoundaries),
 	bMergeBonesOnCoincidentVerts(bMergeCoincidentVertBones),
+	bUseLegacyAttrGrad(bUseLegacyAttrGradient),
 	MeshManager(InSrcVerts, InNumSrcVerts, InSrcIndexes, InNumSrcIndexes, bMergeCoincidentVertBones)
 {
 
diff --git a/Engine/Plugins/Experimental/SkeletalReduction/Source/Private/SkeletalSimplifier.h b/Engine/Plugins/Experimental/SkeletalReduction/Source/Private/SkeletalSimplifier.h
--- a/Engine/Plugins/Experimental/SkeletalReduction/Source/Private/SkeletalSimplifier.h
+++ b/Engine/Plugins/Experimental/SkeletalReduction/Source/Private/SkeletalSimplifier.h
@@ -136,7 +136,8 @@ namespace SkeletalSimplifier
 			const float VolumeImportance,
 			const bool bVolumePreservation,
 			const bool bEnforceBoundaries,
-			const bool  bMergeCoincidentVertBones);
+			const bool bMergeCoincidentVertBones,
+			const bool bUseLegacyAttrGradient);
 
 		~FMeshSimplifier();
 
@@ -290,6 +291,9 @@ namespace SkeletalSimplifier
 
 		// --- End Magic numbers
 
+		// forces the simplifier to use the old attribute interpolation gradient method.  Needed for backward compatibilty with certain cloth assets
+		bool bUseLegacyAttrGrad;
+
 
 		// Heap that maps edges to the cost of collapse
 		FBinaryHeap<double>		CollapseCostHeap;
@@ -336,7 +340,7 @@ namespace SkeletalSimplifier
 		{
 			return FMeshSimplifier::WedgeQuadricType(
 				tri.verts[0]->vert, tri.verts[1]->vert, tri.verts[2]->vert,
-				this->BasicAttrWeights, this->AdditionalAttrWeights);
+				this->BasicAttrWeights, this->AdditionalAttrWeights, bUseLegacyAttrGrad);
 		};
 
 		return QuadricCache.GetWedgeQuadric(v, TriQuadricFatory);
diff --git a/Engine/Plugins/Experimental/SkeletalReduction/Source/Private/SkeletalSimplifierLinearAlgebra.h b/Engine/Plugins/Experimental/SkeletalReduction/Source/Private/SkeletalSimplifierLinearAlgebra.h
--- a/Engine/Plugins/Experimental/SkeletalReduction/Source/Private/SkeletalSimplifierLinearAlgebra.h
+++ b/Engine/Plugins/Experimental/SkeletalReduction/Source/Private/SkeletalSimplifierLinearAlgebra.h
@@ -1343,7 +1343,11 @@ namespace SkeletalSimplifier
 
 		/**
 		* This class generates the interpolation coefficients vector (Vec3d) g  and distance (double) d
-		* defined over the face of a triangle.
+		* defined over the face of a triangle.  It does this in one of two ways.  The first constructs a local basis relative to the 
+		* triangle.  This is preferable.
+		* 
+		* The second is the so-called "legacy method" it uses position vectors for the vertices and will fall back to a crude approximation
+		* when the triangle lies on a plane that intersects the origin.
 		*
 		* *
 		* The actual system solved is:
@@ -1379,26 +1383,91 @@ namespace SkeletalSimplifier
 		*
 		*
 		* The computation is broken up do allow for reuse with multiple sets of per-vertex data.
+		* 
+		* ------------------------------------------
+		* Alternately when using the legacy method
+		* ------------------------------------------
+		* The Position Matrix is defined in terms of the three corners of triangle
+		*  {pa, pb, pc} with corresponding normal 'FaceNormal'
+		*
+		* Position Matrix =
+		*   ( pa_0, pa_1, pa_2 )
+		*   ( pb_0, pb_1, pb_2 )
+		*   ( pc_0, pc_1, pc_2 )
+		*
+		* The actual system solved is:
+		*  <pa | g> + d = s0
+		*  <pb | g> + d = s1
+		*  <pc | g> + d = s2
+		*  <FaceNormal | g> = 0
+		*
+		*
+		* In matrix form   ( PositionMatrix   Vec3(1) )   ( g )  = ( s )
+		*                  ( FaceNormal^T ,     0     )   ( d )    ( 0 )
+		*
+		* where the vector (Vec3d) 's' represents the per-vertex data that forms boundary conditions
+		* for the interpolation.
+		*
+		*  The actual solution is given by:
+		*  -- Distance
+		*  double d = < FaceNormal | InvsPositionMatrix * s> / <FaceNormal | InversePositionMatrix * Vec3d(1) >;
+		*  -- Gradient
+		*  Vec3d  g = InversePositionMatrix * s - d * InversePositionMatrix * Vec3d(1);
+		*
+		* The computation is broken up do allow for reuse with multiple sets of per-vertex data.
 		*/
 		class InverseGradientProjection
 		{
 		public:
-			InverseGradientProjection(const Vec3d& Pos0, const Vec3d& Pos1, const Vec3d& Pos2, const Vec3d& FaceNormal)
+			InverseGradientProjection(const Vec3d& Pos0, const Vec3d& Pos1, const Vec3d& Pos2, const Vec3d& FaceNormal, bool bUseLegacyGradientMethod)
+			: bLegacyGradientMethod(bUseLegacyGradientMethod)
 			{
-
 				// Threshold for computing the matrix inverse.
 				constexpr  double DetThreshold = 1.e-8;
 
-				Origin = Pos0;
+				if (bLegacyGradientMethod)
+				{
+					// re-use matrix and vector for the legacy version
+					DMatrix& PosInv = CoBasisMatrix;
+					Vec3d& Dhat = Origin; // n * Inverse(PositinoMatrix) / < n | Inverse(PositionMatrix) Vec3(1) >
+
+					// row initialization 
+					DMatrix PositionMatrix(Pos0, Pos1, Pos2);
 
-				const Vec3d Side20 = Pos2 - Pos0;
-				const Vec3d Side10 = Pos1 - Pos0;
+					// Compute the inverse of the position matrix
+					PosInv = PositionMatrix.Inverse(bIsValid, DetThreshold);
+					if (bIsValid)
+					{
+						// InversePositionMatrix * Vec3(1)
+						MInv1 = PosInv.RowSum();
 
-				DMatrix BasisMatrix(Side20,
-									Side10,
-									FaceNormal);
+						// <FaceNormal | InversePositionMatrix 
+						Dhat = FaceNormal * PosInv;
 
-				CoBasisMatrix = BasisMatrix.Inverse(bIsValid, DetThreshold);
+						// <FaceNormal | InvesePositionMatrix Vec3d(1)> 
+						double ReScale = Dhat[0] + Dhat[1] + Dhat[2];
+
+						bIsValid = bIsValid && (FMath::Abs(ReScale) > 1.e-8);
+
+						// divide by <FaceNormal | InvesePositionMatrix Vec3d(1)> 
+						Dhat *= (1. / ReScale);
+						//now:  Dhat =  <FaceNormal | InvPos  / <FaceNormal | InvPos. Vec3d(1) >
+					}
+
+				}
+				else // the standard method
+				{ 
+					Origin = Pos0;
+
+					const Vec3d Side20 = Pos2 - Pos0;
+					const Vec3d Side10 = Pos1 - Pos0;
+
+					DMatrix BasisMatrix(Side20,
+										Side10,
+										FaceNormal);
+
+					CoBasisMatrix = BasisMatrix.Inverse(bIsValid, DetThreshold);
+				}
 			}
 
 
@@ -1412,19 +1481,46 @@ namespace SkeletalSimplifier
 			// @return  Distance
 			double ComputeGradient(const Vec3d& PerVertexData, Vec3d& OutGradient) const
 			{
-				Vec3d DataVec(PerVertexData[2] - PerVertexData[0], PerVertexData[1] - PerVertexData[0], 0.);
+				if (bLegacyGradientMethod)
+				{
+					// re-use matrix and vector for the legacy version
+					const DMatrix& PosInv = CoBasisMatrix;
+					const Vec3d& Dhat = Origin;
 
-				OutGradient = CoBasisMatrix * DataVec;
-				double Distance = PerVertexData[0] - Origin.DotProduct(OutGradient);
+					// PosInv . s
+					Vec3d MInvS = PosInv * PerVertexData;
 
-				return Distance;
+					// Dist = <dhat | s>
+					double Distance = Dhat.DotProduct(PerVertexData);
+
+					// Grad =  PosInv . s - Dist * PosInv . (1, 1, 1} 
+					OutGradient = MInvS - Distance * MInv1;
+
+					return Distance;
+				}
+				else // the standard method
+				{ 
+					Vec3d DataVec(PerVertexData[2] - PerVertexData[0], PerVertexData[1] - PerVertexData[0], 0.);
+
+					OutGradient = CoBasisMatrix * DataVec;
+					double Distance = PerVertexData[0] - Origin.DotProduct(OutGradient);
+
+					return Distance;
+				}
 			}
 
 		private:
 
 			bool    bIsValid;  // Indicates if the inversions incurred a divide by very-small-number.
+		
+			bool    bLegacyGradientMethod = false;
+
+			// used by the standard method
 			DMatrix CoBasisMatrix;   // Inverse( {P2-P0}, {P1 - P0}, Normal}
-			Vec3d   Origin;    // corresponds to P0
+			Vec3d   Origin;          // corresponds to P0
+
+			// used by the legacy version only
+			Vec3d   MInv1;                   // Inverse(PositionMatrix) * Vec3(1)
 		};
 
 		template <typename SparseVecDOrDenseVecD>
diff --git a/Engine/Plugins/Experimental/SkeletalReduction/Source/Private/SkeletalSimplifierQuadrics.h b/Engine/Plugins/Experimental/SkeletalReduction/Source/Private/SkeletalSimplifierQuadrics.h
--- a/Engine/Plugins/Experimental/SkeletalReduction/Source/Private/SkeletalSimplifierQuadrics.h
+++ b/Engine/Plugins/Experimental/SkeletalReduction/Source/Private/SkeletalSimplifierQuadrics.h
@@ -194,7 +194,8 @@ namespace SkeletalSimplifier
 				         const SimpVertexType& TriVert1,
 				         const SimpVertexType& TriVert2,
 				         const D1VectorType& BasicWeights,
-				         const SparseWeightContainerType& AdditionalWeights);
+				         const SparseWeightContainerType& AdditionalWeights,
+						 const bool bUseLegacyAttrGrad);
 
 
 			TFaceQuadric() = default;
@@ -550,7 +551,8 @@ namespace SkeletalSimplifier
 
 		template <typename SimpVertexType, typename SparseWeightArrayType>
 		TFaceQuadric<SimpVertexType, SparseWeightArrayType>::TFaceQuadric(const SimpVertexType& TriVert0, const SimpVertexType& TriVert1, const SimpVertexType& TriVert2,
-			                                                              const D1VectorType& BasicWeights, const SparseWeightContainerType& AdditionalWeights)
+			                                                              const D1VectorType& BasicWeights, const SparseWeightContainerType& AdditionalWeights,
+																		  const bool bUseLegacyAttrGrad)
 		{
 
 			// Convert point locations to double.
@@ -608,7 +610,8 @@ namespace SkeletalSimplifier
 			InverseGradientProjection  GradientTool(Vert0Pos,
 													Vert1Pos,
 													Vert2Pos,
-													FaceNormal);
+													FaceNormal,
+													bUseLegacyAttrGrad);
 
 			// Accumulate the terms related to the gradient of the Basic Attributes
 			{
diff --git a/Engine/Source/Editor/DetailCustomizations/Private/SkeletalMeshReductionSettingsDetails.cpp b/Engine/Source/Editor/DetailCustomizations/Private/SkeletalMeshReductionSettingsDetails.cpp
--- a/Engine/Source/Editor/DetailCustomizations/Private/SkeletalMeshReductionSettingsDetails.cpp
+++ b/Engine/Source/Editor/DetailCustomizations/Private/SkeletalMeshReductionSettingsDetails.cpp
@@ -200,7 +200,8 @@ void FSkeletalMeshReductionSettingsDetails::CustomizeChildren(TSharedRef<IProper
 			GET_MEMBER_NAME_CHECKED(FSkeletalMeshOptimizationSettings, bEnforceBoneBoundaries),
 			GET_MEMBER_NAME_CHECKED(FSkeletalMeshOptimizationSettings, bMergeCoincidentVertBones),
 			GET_MEMBER_NAME_CHECKED(FSkeletalMeshOptimizationSettings, VolumeImportance),
-			GET_MEMBER_NAME_CHECKED(FSkeletalMeshOptimizationSettings, bLockColorBounaries)
+			GET_MEMBER_NAME_CHECKED(FSkeletalMeshOptimizationSettings, bLockColorBounaries),
+			GET_MEMBER_NAME_CHECKED(FSkeletalMeshOptimizationSettings, bImproveTrianglesForCloth)
 		};
 
 		uint32 NumChildren = 0;
diff --git a/Engine/Source/Editor/Persona/Private/PersonaMeshDetails.cpp b/Engine/Source/Editor/Persona/Private/PersonaMeshDetails.cpp
--- a/Engine/Source/Editor/Persona/Private/PersonaMeshDetails.cpp
+++ b/Engine/Source/Editor/Persona/Private/PersonaMeshDetails.cpp
@@ -792,6 +792,14 @@ void FSkeletalMeshReductionSettingsLayout::GenerateChildContent(IDetailChildrenB
 			TEXT("LockColorBoundaries"),
 			FGetCheckBoxStateDelegate::CreateRaw(this, &FSkeletalMeshReductionSettingsLayout::GetLockColorBounaries),
 			FSetCheckBoxStateDelegate::CreateRaw(this, &FSkeletalMeshReductionSettingsLayout::SetLockColorBounaries));
+
+		AddBoolRow(ChildrenBuilder,
+			LOCTEXT("ImproveTrianglesForCloth_Row", "ImproveTrianglesForCloth"),
+			LOCTEXT("ImproveTrianglesForCloth_RowNameContent", "Improve Triangles For Cloth"),
+			LOCTEXT("ImproveTrianglesForCloth_RowNameContentTooltip", "Better distribution of triangles on 2d meshes, such as flat cloth, but at the cost of potentially worse UVs in those areas.  This generally has little or no effect for mesh regions that aren't laid out on a plane intersecting the origin such as the xy-plane. When this is disabled, the planar regions may simplify to fewer large triangles."),
+			TEXT("ImproveTrianglesForCloth"),
+			FGetCheckBoxStateDelegate::CreateRaw(this, &FSkeletalMeshReductionSettingsLayout::GetImproveTrianglesForCloth),
+			FSetCheckBoxStateDelegate::CreateRaw(this, &FSkeletalMeshReductionSettingsLayout::SetImproveTrianglesForCloth));
 	}
 
 	AddBaseLODRow(ChildrenBuilder);
@@ -1402,6 +1410,20 @@ void FSkeletalMeshReductionSettingsLayout::SetLockColorBounaries(ECheckBoxState
 	ReductionSettings.bLockColorBounaries = (NewState == ECheckBoxState::Checked) ? true : false;
 }
 
+ECheckBoxState FSkeletalMeshReductionSettingsLayout::GetImproveTrianglesForCloth() const
+{
+	return ReductionSettings.bImproveTrianglesForCloth ? ECheckBoxState::Checked : ECheckBoxState::Unchecked;
+}
+
+void FSkeletalMeshReductionSettingsLayout::SetImproveTrianglesForCloth(ECheckBoxState NewState)
+{
+	FText TransactionText = FText::Format(LOCTEXT("PersonaReductionChangedSetUseLegacyVersionLOD", "LOD{0} reduction settings: improve triangles for cloth changed"), LODIndex);
+	FScopedTransaction Transaction(TransactionText);
+	ModifyMeshLODSettingsDelegate.ExecuteIfBound(LODIndex);
+
+	ReductionSettings.bImproveTrianglesForCloth = (NewState == ECheckBoxState::Checked) ? true : false;
+}
+
 ECheckBoxState FSkeletalMeshReductionSettingsLayout::GetEnforceBoneBoundaries() const
 {
 	return ReductionSettings.bEnforceBoneBoundaries ? ECheckBoxState::Checked : ECheckBoxState::Unchecked;
diff --git a/Engine/Source/Editor/Persona/Private/PersonaMeshDetails.h b/Engine/Source/Editor/Persona/Private/PersonaMeshDetails.h
--- a/Engine/Source/Editor/Persona/Private/PersonaMeshDetails.h
+++ b/Engine/Source/Editor/Persona/Private/PersonaMeshDetails.h
@@ -191,6 +191,9 @@ private:
 	ECheckBoxState GetLockColorBounaries() const;
 	void SetLockColorBounaries(ECheckBoxState NewState);
 
+	ECheckBoxState GetImproveTrianglesForCloth() const;
+	void SetImproveTrianglesForCloth(ECheckBoxState NewState);
+
 	ECheckBoxState GetEnforceBoneBoundaries() const;
 	void SetEnforceBoneBoundaries(ECheckBoxState NewState);
 
diff --git a/Engine/Source/Runtime/Engine/Private/SkinnedAssetCommon.cpp b/Engine/Source/Runtime/Engine/Private/SkinnedAssetCommon.cpp
--- a/Engine/Source/Runtime/Engine/Private/SkinnedAssetCommon.cpp
+++ b/Engine/Source/Runtime/Engine/Private/SkinnedAssetCommon.cpp
@@ -29,6 +29,7 @@ static void SerializeReductionSettingsForDDC(FArchive& Ar, FSkeletalMeshOptimiza
 	FArchive_Serialize_BitfieldBool(Ar, ReductionSettings.bMergeCoincidentVertBones);
 	FArchive_Serialize_BitfieldBool(Ar, ReductionSettings.bLockEdges);
 	FArchive_Serialize_BitfieldBool(Ar, ReductionSettings.bLockColorBounaries);
+	FArchive_Serialize_BitfieldBool(Ar, ReductionSettings.bImproveTrianglesForCloth);
 	Ar << ReductionSettings.TerminationCriterion;
 	Ar << ReductionSettings.NumOfTrianglesPercentage;
 	Ar << ReductionSettings.NumOfVertPercentage;
diff --git a/Engine/Source/Runtime/Engine/Public/SkeletalMeshReductionSettings.h b/Engine/Source/Runtime/Engine/Public/SkeletalMeshReductionSettings.h
--- a/Engine/Source/Runtime/Engine/Public/SkeletalMeshReductionSettings.h
+++ b/Engine/Source/Runtime/Engine/Public/SkeletalMeshReductionSettings.h
@@ -156,10 +156,15 @@ struct FSkeletalMeshOptimizationSettings
 	UPROPERTY(EditAnywhere, Category = FSkeletalMeshOptimizationSettings, meta = (DisplayName = "Lock Vertex Color Boundaries"))
 	uint8 bLockColorBounaries : 1;
 
+	/** Better distribution of triangles on 2d meshes, such as flat cloth, but at the cost of potentially worse UVs in those areas.  This generally has little or no effect for mesh regions that aren't laid out on a plane intersecting the origin such as the xy-plane. When this is disabled, the planar regions may simplify to fewer large triangles.*/
+	UPROPERTY(EditAnywhere, Category = FSkeletalMeshOptimizationSettings)
+	uint8 bImproveTrianglesForCloth : 1;
+
 	/** Base LOD index to generate this LOD. By default, we generate from LOD 0 */
 	UPROPERTY(EditAnywhere, Category = FSkeletalMeshOptimizationSettings)
 	int32 BaseLOD;
 
+
 #if WITH_EDITORONLY_DATA
 	UPROPERTY()
 	TArray<FBoneReference> BonesToRemove_DEPRECATED;
@@ -197,6 +202,7 @@ struct FSkeletalMeshOptimizationSettings
 		, VolumeImportance(1.f)
 		, bLockEdges(false)
 		, bLockColorBounaries(false)
+		, bImproveTrianglesForCloth(true)
 		, BaseLOD(0)
 #if WITH_EDITORONLY_DATA
 		, BakePose_DEPRECATED(nullptr)
