SkeletalMeshReduction: Disable core.uninitialized.ArraySubscript clang checker

#rnx

diff --git a/Engine/Plugins/Experimental/SkeletalReduction/Source/SkeletalMeshReduction.Build.cs b/Engine/Plugins/Experimental/SkeletalReduction/Source/SkeletalMeshReduction.Build.cs
--- a/Engine/Plugins/Experimental/SkeletalReduction/Source/SkeletalMeshReduction.Build.cs
+++ b/Engine/Plugins/Experimental/SkeletalReduction/Source/SkeletalMeshReduction.Build.cs
@@ -6,11 +6,11 @@ namespace UnrealBuildTool.Rules
     {
         public SkeletalMeshReduction(ReadOnlyTargetRules Target) : base(Target)
         {
+			StaticAnalyzerDisabledCheckers.Add("core.uninitialized.ArraySubscript");
 
-
-            // For boost:: and TBB:: code
-            //UndefinedIdentifierWarningLevel = WarningLevel.Off;
-            //bUseRTTI = true;
+			// For boost:: and TBB:: code
+			//UndefinedIdentifierWarningLevel = WarningLevel.Off;
+			//bUseRTTI = true;
 			/*
             PublicIncludePaths.AddRange(
                 new string[] {
@@ -24,7 +24,7 @@ namespace UnrealBuildTool.Rules
 				}
                 );
 				*/
-            PublicDependencyModuleNames.AddRange(
+			PublicDependencyModuleNames.AddRange(
                 new string[]
                 {
                     "Core",
